- Add support for collision and multiple sections to ProcMesh

diff --git a/Engine/Plugins/Runtime/ProceduralMeshComponent/Source/ProceduralMeshComponent/Private/ProceduralMeshComponent.cpp b/Engine/Plugins/Runtime/ProceduralMeshComponent/Source/ProceduralMeshComponent/Private/ProceduralMeshComponent.cpp
--- a/Engine/Plugins/Runtime/ProceduralMeshComponent/Source/ProceduralMeshComponent/Private/ProceduralMeshComponent.cpp
+++ b/Engine/Plugins/Runtime/ProceduralMeshComponent/Source/ProceduralMeshComponent/Private/ProceduralMeshComponent.cpp
@@ -76,7 +76,25 @@ public:
 	}
 };
 
-/** Scene proxy */
+/** Class representing a single section of the proc mesh */
+class FProcMeshProxySection
+{
+public:
+	/** Material applied to this section */
+	UMaterialInterface* Material;
+	/** Vertex buffer for this section */
+	FProcMeshVertexBuffer VertexBuffer;
+	/** Index buffer for this section */
+	FProcMeshIndexBuffer IndexBuffer;
+	/** Vertex factory for this section */
+	FProcMeshVertexFactory VertexFactory;
+
+	FProcMeshProxySection()
+	: Material(NULL)
+	{}
+};
+
+/** Procedural mesh scene proxy */
 class FProceduralMeshSceneProxy : public FPrimitiveSceneProxy
 {
 public:
@@ -85,102 +103,135 @@ public:
 		: FPrimitiveSceneProxy(Component)
 		, MaterialRelevance(Component->GetMaterialRelevance(GetScene().GetFeatureLevel()))
 	{
-		// Copy data from vertex buffer
-		const int32 NumVerts = Component->ProcVertexBuffer.Num();
-
-
-		// Allocate verts
-		VertexBuffer.Vertices.SetNumUninitialized(NumVerts);
-		// Copy verts
-		for(int VertIdx=0; VertIdx<NumVerts; VertIdx++)
-		{
-			FProcMeshVertex& ProcVert = Component->ProcVertexBuffer[VertIdx];
-
-			FDynamicMeshVertex& Vert = VertexBuffer.Vertices[VertIdx];
-			Vert.Position = ProcVert.Position;
-			Vert.Color = ProcVert.Color;
-			Vert.TextureCoordinate = ProcVert.UV0;
-			Vert.TangentX = ProcVert.Tangent.TangentX;
-			Vert.TangentZ = ProcVert.Normal;
-			Vert.TangentZ.Vector.W = ProcVert.Tangent.bFlipTangentY ? 0 : 255;
-		}
-
-		IndexBuffer.Indices = Component->ProcIndexBuffer;
-
-		// Init vertex factory
-		VertexFactory.Init(&VertexBuffer);
-
-		// Enqueue initialization of render resource
-		BeginInitResource(&VertexBuffer);
-		BeginInitResource(&IndexBuffer);
-		BeginInitResource(&VertexFactory);
-
-		// Grab material
-		Material = Component->GetMaterial(0);
-		if(Material == NULL)
+		// Copy each section
+		const int32 NumSections = Component->ProcMeshSections.Num();
+		Sections.AddZeroed(NumSections);
+		for (int SectionIdx = 0; SectionIdx < NumSections; SectionIdx++)
 		{
-			Material = UMaterial::GetDefaultMaterial(MD_Surface);
+			FProcMeshSection& SrcSection = Component->ProcMeshSections[SectionIdx];
+			if(SrcSection.ProcIndexBuffer.Num() > 0 && SrcSection.ProcVertexBuffer.Num() > 0)
+			{
+				FProcMeshProxySection* NewSection = new FProcMeshProxySection();
+
+				// Copy data from vertex buffer
+				const int32 NumVerts = SrcSection.ProcVertexBuffer.Num();
+
+				// Allocate verts
+				NewSection->VertexBuffer.Vertices.SetNumUninitialized(NumVerts);
+				// Copy verts
+				for (int VertIdx = 0; VertIdx < NumVerts; VertIdx++)
+				{
+					FProcMeshVertex& ProcVert = SrcSection.ProcVertexBuffer[VertIdx];
+
+					FDynamicMeshVertex& Vert = NewSection->VertexBuffer.Vertices[VertIdx];
+					Vert.Position = ProcVert.Position;
+					Vert.Color = ProcVert.Color;
+					Vert.TextureCoordinate = ProcVert.UV0;
+					Vert.TangentX = ProcVert.Tangent.TangentX;
+					Vert.TangentZ = ProcVert.Normal;
+					Vert.TangentZ.Vector.W = ProcVert.Tangent.bFlipTangentY ? 0 : 255;
+				}
+
+				// Copy index buffer
+				NewSection->IndexBuffer.Indices = SrcSection.ProcIndexBuffer;
+
+				// Init vertex factory
+				NewSection->VertexFactory.Init(&NewSection->VertexBuffer);
+
+				// Enqueue initialization of render resource
+				BeginInitResource(&NewSection->VertexBuffer);
+				BeginInitResource(&NewSection->IndexBuffer);
+				BeginInitResource(&NewSection->VertexFactory);
+
+				// Grab material
+				NewSection->Material = Component->GetMaterial(SectionIdx);
+				if (NewSection->Material == NULL)
+				{
+					NewSection->Material = UMaterial::GetDefaultMaterial(MD_Surface);
+				}
+
+				// Save ref to new section
+				Sections[SectionIdx] = NewSection;
+			}
 		}
 	}
 
 	virtual ~FProceduralMeshSceneProxy()
 	{
-		VertexBuffer.ReleaseResource();
-		IndexBuffer.ReleaseResource();
-		VertexFactory.ReleaseResource();
+		for (FProcMeshProxySection* Section : Sections)
+		{
+			if(Section != nullptr)
+			{
+				Section->VertexBuffer.ReleaseResource();
+				Section->IndexBuffer.ReleaseResource();
+				Section->VertexFactory.ReleaseResource();
+				delete Section;
+			}
+		}
 	}
 
 	virtual void GetDynamicMeshElements(const TArray<const FSceneView*>& Views, const FSceneViewFamily& ViewFamily, uint32 VisibilityMap, FMeshElementCollector& Collector) const override
 	{
+		// Set up wireframe material (if needed)
 		const bool bWireframe = AllowDebugViewmodes() && ViewFamily.EngineShowFlags.Wireframe;
 
-		auto WireframeMaterialInstance = new FColoredMaterialRenderProxy(
-			GEngine->WireframeMaterial ? GEngine->WireframeMaterial->GetRenderProxy(IsSelected()) : NULL,
-			FLinearColor(0, 0.5f, 1.f)
-			);
-
-		Collector.RegisterOneFrameMaterialProxy(WireframeMaterialInstance);
-
-		FMaterialRenderProxy* MaterialProxy = NULL;
-		if(bWireframe)
+		FColoredMaterialRenderProxy* WireframeMaterialInstance = NULL;
+		if (bWireframe)
 		{
-			MaterialProxy = WireframeMaterialInstance;
+			WireframeMaterialInstance = new FColoredMaterialRenderProxy(
+				GEngine->WireframeMaterial ? GEngine->WireframeMaterial->GetRenderProxy(IsSelected()) : NULL,
+				FLinearColor(0, 0.5f, 1.f)
+				);
+
+			Collector.RegisterOneFrameMaterialProxy(WireframeMaterialInstance);
 		}
-		else
+
+		// Iterate over sections
+		for (const FProcMeshProxySection* Section : Sections)
 		{
-			MaterialProxy = Material->GetRenderProxy(IsSelected());
+			if(Section != nullptr)
+			{
+				FMaterialRenderProxy* MaterialProxy = bWireframe ? WireframeMaterialInstance : Section->Material->GetRenderProxy(IsSelected());
+
+				// For each view..
+				for (int32 ViewIndex = 0; ViewIndex < Views.Num(); ViewIndex++)
+				{
+					if (VisibilityMap & (1 << ViewIndex))
+					{
+						const FSceneView* View = Views[ViewIndex];
+						// Draw the mesh.
+						FMeshBatch& Mesh = Collector.AllocateMesh();
+						FMeshBatchElement& BatchElement = Mesh.Elements[0];
+						BatchElement.IndexBuffer = &Section->IndexBuffer;
+						Mesh.bWireframe = bWireframe;
+						Mesh.VertexFactory = &Section->VertexFactory;
+						Mesh.MaterialRenderProxy = MaterialProxy;
+						BatchElement.PrimitiveUniformBuffer = CreatePrimitiveUniformBufferImmediate(GetLocalToWorld(), GetBounds(), GetLocalBounds(), true, UseEditorDepthTest());
+						BatchElement.FirstIndex = 0;
+						BatchElement.NumPrimitives = Section->IndexBuffer.Indices.Num() / 3;
+						BatchElement.MinVertexIndex = 0;
+						BatchElement.MaxVertexIndex = Section->VertexBuffer.Vertices.Num() - 1;
+						Mesh.ReverseCulling = IsLocalToWorldDeterminantNegative();
+						Mesh.Type = PT_TriangleList;
+						Mesh.DepthPriorityGroup = SDPG_World;
+						Mesh.bCanApplyViewModeOverrides = false;
+						Collector.AddMesh(ViewIndex, Mesh);
+					}
+				}
+			}
 		}
 
+		// Draw bounds
+#if !(UE_BUILD_SHIPPING || UE_BUILD_TEST)
 		for (int32 ViewIndex = 0; ViewIndex < Views.Num(); ViewIndex++)
 		{
 			if (VisibilityMap & (1 << ViewIndex))
 			{
-				const FSceneView* View = Views[ViewIndex];
-				// Draw the mesh.
-				FMeshBatch& Mesh = Collector.AllocateMesh();
-				FMeshBatchElement& BatchElement = Mesh.Elements[0];
-				BatchElement.IndexBuffer = &IndexBuffer;
-				Mesh.bWireframe = bWireframe;
-				Mesh.VertexFactory = &VertexFactory;
-				Mesh.MaterialRenderProxy = MaterialProxy;
-				BatchElement.PrimitiveUniformBuffer = CreatePrimitiveUniformBufferImmediate(GetLocalToWorld(), GetBounds(), GetLocalBounds(), true, UseEditorDepthTest());
-				BatchElement.FirstIndex = 0;
-				BatchElement.NumPrimitives = IndexBuffer.Indices.Num() / 3;
-				BatchElement.MinVertexIndex = 0;
-				BatchElement.MaxVertexIndex = VertexBuffer.Vertices.Num() - 1;
-				Mesh.ReverseCulling = IsLocalToWorldDeterminantNegative();
-				Mesh.Type = PT_TriangleList;
-				Mesh.DepthPriorityGroup = SDPG_World;
-				Mesh.bCanApplyViewModeOverrides = false;
-				Collector.AddMesh(ViewIndex, Mesh);
-
-#if !(UE_BUILD_SHIPPING || UE_BUILD_TEST)
 				// Render bounds
 				RenderBounds(Collector.GetPDI(ViewIndex), ViewFamily.EngineShowFlags, GetBounds(), IsSelected());
-#endif
 			}
 		}
-
+#endif
 	}
 
 	virtual FPrimitiveViewRelevance GetViewRelevance(const FSceneView* View)
@@ -198,16 +249,19 @@ public:
 		return !MaterialRelevance.bDisableDepthTest;
 	}
 
-	virtual uint32 GetMemoryFootprint( void ) const { return( sizeof( *this ) + GetAllocatedSize() ); }
+	virtual uint32 GetMemoryFootprint( void ) const 
+	{ 
+		return( sizeof( *this ) + GetAllocatedSize() ); 
+	}
 
-	uint32 GetAllocatedSize( void ) const { return( FPrimitiveSceneProxy::GetAllocatedSize() ); }
+	uint32 GetAllocatedSize( void ) const 
+	{ 
+		return( FPrimitiveSceneProxy::GetAllocatedSize() ); 
+	}
 
 private:
-
-	UMaterialInterface* Material;
-	FProcMeshVertexBuffer VertexBuffer;
-	FProcMeshIndexBuffer IndexBuffer;
-	FProcMeshVertexFactory VertexFactory;
+	/** Array of sections */
+	TArray<FProcMeshProxySection*> Sections;
 
 	FMaterialRelevance MaterialRelevance;
 };
@@ -218,24 +272,24 @@ private:
 UProceduralMeshComponent::UProceduralMeshComponent( const FObjectInitializer& ObjectInitializer )
 	: Super( ObjectInitializer )
 {
-	PrimaryComponentTick.bCanEverTick = false;
-
-	// DEfault to no collision. Building collision is expensive, make sure user opts into it.
-	SetCollisionProfileName(UCollisionProfile::NoCollision_ProfileName);
 }
 
-void UProceduralMeshComponent::CreateMesh(const TArray<FVector>& Vertices, const TArray<int32>& Triangles, const TArray<FVector>& Normals, const TArray<FVector2D>& UV0, const TArray<FColor>& VertexColors, const TArray<FProcMeshTangent>& Tangents)
+void UProceduralMeshComponent::CreateMeshSection(int32 SectionIndex, const TArray<FVector>& Vertices, const TArray<int32>& Triangles, const TArray<FVector>& Normals, const TArray<FVector2D>& UV0, const TArray<FColor>& VertexColors, const TArray<FProcMeshTangent>& Tangents, bool bCreateCollision)
 {
-	// Local bounding box
-	FBox LocalBox(0);
+	// Ensure sections array is long enough
+	ProcMeshSections.SetNum(SectionIndex+1, false);
+
+	// Reset this section (in case it already existed)
+	FProcMeshSection& NewSection = ProcMeshSections[SectionIndex];
+	NewSection.Reset();
 
 	// Copy data to vertex buffer
 	const int32 NumVerts = Vertices.Num();
-	ProcVertexBuffer.Reset();
-	ProcVertexBuffer.AddUninitialized(NumVerts);
+	NewSection.ProcVertexBuffer.Reset();
+	NewSection.ProcVertexBuffer.AddUninitialized(NumVerts);
 	for(int32 VertIdx=0; VertIdx<NumVerts; VertIdx++)
 	{
-		FProcMeshVertex& Vertex = ProcVertexBuffer[VertIdx];
+		FProcMeshVertex& Vertex = NewSection.ProcVertexBuffer[VertIdx];
 
 		Vertex.Position = Vertices[VertIdx];
 		Vertex.Normal = (Normals.Num() == NumVerts) ? Normals[VertIdx] : FVector(0.f, 0.f, 1.f);
@@ -244,45 +298,169 @@ void UProceduralMeshComponent::CreateMesh(const TArray<FVector>& Vertices, const
 		Vertex.Tangent = (Tangents.Num() == NumVerts) ? Tangents[VertIdx] : FProcMeshTangent();
 
 		// Update bounding box
-		LocalBox += Vertex.Position;
+		NewSection.SectionLocalBox += Vertex.Position;
 	}
 
 	// Copy index buffer (clamping to vertex range)
 	const int32 NumTriIndices = Triangles.Num();
-	ProcIndexBuffer.Reset();
-	ProcIndexBuffer.AddUninitialized(NumTriIndices);
+	NewSection.ProcIndexBuffer.Reset();
+	NewSection.ProcIndexBuffer.AddUninitialized(NumTriIndices);
 	for (int32 IndexIdx = 0; IndexIdx < NumTriIndices; IndexIdx++)	
 	{
-		ProcIndexBuffer[IndexIdx] = FMath::Min(Triangles[IndexIdx], NumVerts - 1);
+		NewSection.ProcIndexBuffer[IndexIdx] = FMath::Min(Triangles[IndexIdx], NumVerts - 1);
 	}
 
-	// Set bounds
-	LocalBounds = LocalBox;
+	NewSection.bEnableCollision = bCreateCollision;
 
-	// Recreate scene proxy
-	MarkRenderStateDirty();
+	UpdateLocalBounds(); // Update overall bounds
+	UpdateCollision(); // Mark collision as dirty
+	MarkRenderStateDirty(); // Recreate scene proxy
 }
 
+void UProceduralMeshComponent::ClearMeshSection(int32 SectionIndex)
+{
+	if(SectionIndex < ProcMeshSections.Num())
+	{
+		ProcMeshSections[SectionIndex].Reset();
+		UpdateLocalBounds();
+		UpdateCollision();
+		MarkRenderStateDirty();
+	}
+}
 
-FPrimitiveSceneProxy* UProceduralMeshComponent::CreateSceneProxy()
+void UProceduralMeshComponent::ClearAllMeshSections()
 {
-	FPrimitiveSceneProxy* Proxy = NULL;
-	if(ProcVertexBuffer.Num() > 0 && ProcIndexBuffer.Num() > 0)
+	ProcMeshSections.Empty();
+	UpdateLocalBounds();
+	UpdateCollision();
+	MarkRenderStateDirty();
+}
+
+void UProceduralMeshComponent::UpdateLocalBounds()
+{
+	FBox LocalBox(0);
+
+	for(const FProcMeshSection& Section : ProcMeshSections)
 	{
-		Proxy = new FProceduralMeshSceneProxy(this);
+		LocalBox += Section.SectionLocalBox;
 	}
-	return Proxy;
+
+	LocalBounds = LocalBox.IsValid ? FBoxSphereBounds(LocalBox) : FBoxSphereBounds(FVector(0,0,0), FVector(0,0,0),0); // fallback to reset box sphere bounds
 }
 
-int32 UProceduralMeshComponent::GetNumMaterials() const
+FPrimitiveSceneProxy* UProceduralMeshComponent::CreateSceneProxy()
 {
-	return 1;
+	return new FProceduralMeshSceneProxy(this);
 }
 
+int32 UProceduralMeshComponent::GetNumMaterials() const
+{
+	return ProcMeshSections.Num();
+}
 
 FBoxSphereBounds UProceduralMeshComponent::CalcBounds(const FTransform& LocalToWorld) const
 {
 	return LocalBounds.TransformBy(LocalToWorld);
 }
 
+bool UProceduralMeshComponent::GetPhysicsTriMeshData(struct FTriMeshCollisionData* CollisionData, bool InUseAllTriData)
+{
+	int32 VertexBase = 0; // Base vertex index for current section
+
+	// For each section..
+	for (int32 SectionIdx = 0; SectionIdx < ProcMeshSections.Num(); SectionIdx++)
+	{
+		FProcMeshSection& Section = ProcMeshSections[SectionIdx];
+		// Do we have collision enabled?
+		if(Section.bEnableCollision)
+		{
+			// Copy vert data
+			for (int32 VertIdx = 0; VertIdx < Section.ProcVertexBuffer.Num(); VertIdx++)
+			{
+				CollisionData->Vertices.Add(Section.ProcVertexBuffer[VertIdx].Position);
+			}
+
+			// Copy triangle data
+			const int32 NumTriangles = Section.ProcIndexBuffer.Num() / 3;
+			for (int32 TriIdx = 0; TriIdx < NumTriangles; TriIdx++)
+			{
+				// Need to add base offset for indices
+				FTriIndices Triangle;
+				Triangle.v0 = Section.ProcIndexBuffer[(TriIdx * 3) + 0] + VertexBase;
+				Triangle.v1 = Section.ProcIndexBuffer[(TriIdx * 3) + 1] + VertexBase;
+				Triangle.v2 = Section.ProcIndexBuffer[(TriIdx * 3) + 2] + VertexBase;
+				CollisionData->Indices.Add(Triangle);
+
+				// Also store matrial info
+				CollisionData->MaterialIndices.Add(SectionIdx);
+			}
+
+			// Remember the base index that new verts will be added from in next section
+			VertexBase = CollisionData->Vertices.Num();
+		}
+	}
+
+	CollisionData->bFlipNormals = true;
+
+	return true;
+}
+
+bool UProceduralMeshComponent::ContainsPhysicsTriMeshData(bool InUseAllTriData) const
+{
+	for (const FProcMeshSection& Section : ProcMeshSections)
+	{
+		if(Section.ProcIndexBuffer.Num() >= 3 && Section.bEnableCollision)
+		{
+			return true;
+		}
+	}
+
+	return false;
+}
+
+void UProceduralMeshComponent::CreateProcMeshBodySetup()
+{
+	if (ProcMeshBodySetup == NULL)
+	{
+		ProcMeshBodySetup = NewObject<UBodySetup>(this);
+		ProcMeshBodySetup->BodySetupGuid = FGuid::NewGuid();
+
+		ProcMeshBodySetup->CollisionTraceFlag = CTF_UseComplexAsSimple;
+		ProcMeshBodySetup->bGenerateMirroredCollision = false;
+		ProcMeshBodySetup->bDoubleSidedGeometry = true;
+	}
+}
+
+void UProceduralMeshComponent::UpdateCollision()
+{
+	bool bCreatePhysState = false; // Should we create physics state at the end of this function?
+
+	// If its created, shut it down now
+	if (bPhysicsStateCreated)
+	{
+		DestroyPhysicsState();
+		bCreatePhysState = true;
+	}
+
+	// Ensure we have a BodySetup
+	CreateProcMeshBodySetup();
+
+	// Clear current mesh data
+	ProcMeshBodySetup->InvalidatePhysicsData();
+	// Create new mesh data
+	ProcMeshBodySetup->CreatePhysicsMeshes();
+
+	// Create new instance state if desired
+	if(bCreatePhysState)
+	{
+		CreatePhysicsState();
+	}
+}
+
+UBodySetup* UProceduralMeshComponent::GetBodySetup()
+{
+	CreateProcMeshBodySetup();
+	return ProcMeshBodySetup;
+}
+
 
diff --git a/Engine/Plugins/Runtime/ProceduralMeshComponent/Source/ProceduralMeshComponent/Public/ProceduralMeshComponent.h b/Engine/Plugins/Runtime/ProceduralMeshComponent/Source/ProceduralMeshComponent/Public/ProceduralMeshComponent.h
--- a/Engine/Plugins/Runtime/ProceduralMeshComponent/Source/ProceduralMeshComponent/Public/ProceduralMeshComponent.h
+++ b/Engine/Plugins/Runtime/ProceduralMeshComponent/Source/ProceduralMeshComponent/Public/ProceduralMeshComponent.h
@@ -62,21 +62,70 @@ struct FProcMeshVertex
 	{}
 };
 
+USTRUCT()
+struct FProcMeshSection
+{
+	GENERATED_USTRUCT_BODY()
+
+	/** Vertex buffer for proc mesh */
+	UPROPERTY()
+	TArray<FProcMeshVertex> ProcVertexBuffer;
+	/** Index buffer for proc mesh */
+	UPROPERTY()
+	TArray<int32> ProcIndexBuffer;
+	/** Local bounds of mesh */
+	UPROPERTY()
+	FBox SectionLocalBox;
+	/** Should we build collision data for traingles in this section */
+	UPROPERTY()
+	bool bEnableCollision;
+
+	FProcMeshSection()
+	: SectionLocalBox(0)
+	, bEnableCollision(false)
+	{}
+
+	/** Reset this section, clear all mesh info. */
+	void Reset()
+	{
+		ProcVertexBuffer.Empty();
+		ProcIndexBuffer.Empty();
+		SectionLocalBox.Init();
+		bEnableCollision = false;
+	}
+};
+
 /** Component that allows you to specify custom triangle mesh geometry */
 UCLASS(hidecategories=(Object,LOD), meta=(BlueprintSpawnableComponent), ClassGroup=Rendering)
-class PROCEDURALMESHCOMPONENT_API UProceduralMeshComponent : public UMeshComponent
+class PROCEDURALMESHCOMPONENT_API UProceduralMeshComponent : public UMeshComponent, public IInterface_CollisionDataProvider
 {
 	GENERATED_UCLASS_BODY()
 
 	/** Add triangle to the geometry to use on this triangle mesh. Must call BuildMesh before geometry can be seen. */
 	UFUNCTION(BlueprintCallable, Category = "Components|ProceduralMesh", meta=(AutoCreateRefTerm = "Normals,UV0,VertexColors,Tangents" ))
-	void CreateMesh(const TArray<FVector>& Vertices, const TArray<int32>& Triangles, const TArray<FVector>& Normals, const TArray<FVector2D>& UV0, const TArray<FColor>& VertexColors, const TArray<FProcMeshTangent>& Tangents);
+	void CreateMeshSection(int32 SectionIndex, const TArray<FVector>& Vertices, const TArray<int32>& Triangles, const TArray<FVector>& Normals, const TArray<FVector2D>& UV0, const TArray<FColor>& VertexColors, const TArray<FProcMeshTangent>& Tangents, bool bCreateCollision);
 
+	UFUNCTION(BlueprintCallable, Category = "Components|ProceduralMesh")
+	void ClearMeshSection(int32 SectionIndex);
 
-private:
+	UFUNCTION(BlueprintCallable, Category = "Components|ProceduralMesh")
+	void ClearAllMeshSections();
+
+
+	// Begin Interface_CollisionDataProvider Interface
+	virtual bool GetPhysicsTriMeshData(struct FTriMeshCollisionData* CollisionData, bool InUseAllTriData) override;
+	virtual bool ContainsPhysicsTriMeshData(bool InUseAllTriData) const override;
+	virtual bool WantsNegXTriMesh() override{ return false; }
+	// End Interface_CollisionDataProvider Interface
 
+	/** Collision data */
+	UPROPERTY()
+	class UBodySetup* ProcMeshBodySetup;
+
+private:
 	// Begin UPrimitiveComponent interface.
 	virtual FPrimitiveSceneProxy* CreateSceneProxy() override;
+	virtual class UBodySetup* GetBodySetup() override;
 	// End UPrimitiveComponent interface.
 
 	// Begin UMeshComponent interface.
@@ -87,15 +136,22 @@ private:
 	virtual FBoxSphereBounds CalcBounds(const FTransform& LocalToWorld) const override;
 	// Begin USceneComponent interface.
 
-	// Members below are updated when BuildMesh is called.
+	/** Update LocalBounds member from the local box of each section */
+	void UpdateLocalBounds();
+	/** Ensure ProcMeshBodySetup is allocated and configured */
+	void CreateProcMeshBodySetup();
+	/** Mark collision data as dirty, and re-create on instance if necessary */
+	void UpdateCollision();
 
-	/** Vertex buffer for proc mesh */
-	TArray<FProcMeshVertex> ProcVertexBuffer;
-	/** Index buffer for proc mesh */
-	TArray<int32> ProcIndexBuffer;
-	/** Local bounds of mesh */
+	/** Array of sections of mesh */
+	UPROPERTY()
+	TArray<FProcMeshSection> ProcMeshSections;
+
+	/** Local space bounds of mesh */
+	UPROPERTY()
 	FBoxSphereBounds LocalBounds;
 
+
 	friend class FProceduralMeshSceneProxy;
 };
 
