- Re-enable ProcMeshComp plugin on mobile platforms, but disable recooking mesh collision if !WITH_RUNTIME_PHYSICS_COOKING

diff --git a/Engine/Plugins/Runtime/ProceduralMeshComponent/ProceduralMeshComponent.uplugin b/Engine/Plugins/Runtime/ProceduralMeshComponent/ProceduralMeshComponent.uplugin
--- a/Engine/Plugins/Runtime/ProceduralMeshComponent/ProceduralMeshComponent.uplugin
+++ b/Engine/Plugins/Runtime/ProceduralMeshComponent/ProceduralMeshComponent.uplugin
@@ -16,7 +16,7 @@
 		{
 			"Name" : "ProceduralMeshComponent",
 			"Type" : "Runtime",
-			"BlacklistPlatforms" : [ "HTML5", "IOS", "Android" ]
+			"BlacklistPlatforms" : [ "HTML5" ]
 		}
 	]
 }
\ No newline at end of file
diff --git a/Engine/Plugins/Runtime/ProceduralMeshComponent/Source/ProceduralMeshComponent/Private/ProceduralMeshComponent.cpp b/Engine/Plugins/Runtime/ProceduralMeshComponent/Source/ProceduralMeshComponent/Private/ProceduralMeshComponent.cpp
--- a/Engine/Plugins/Runtime/ProceduralMeshComponent/Source/ProceduralMeshComponent/Private/ProceduralMeshComponent.cpp
+++ b/Engine/Plugins/Runtime/ProceduralMeshComponent/Source/ProceduralMeshComponent/Private/ProceduralMeshComponent.cpp
@@ -445,10 +445,12 @@ void UProceduralMeshComponent::UpdateCollision()
 	// Ensure we have a BodySetup
 	CreateProcMeshBodySetup();
 
+#if WITH_RUNTIME_PHYSICS_COOKING || WITH_EDITOR
 	// Clear current mesh data
 	ProcMeshBodySetup->InvalidatePhysicsData();
 	// Create new mesh data
 	ProcMeshBodySetup->CreatePhysicsMeshes();
+#endif // WITH_RUNTIME_PHYSICS_COOKING || WITH_EDITOR
 
 	// Create new instance state if desired
 	if(bCreatePhysState)
