UE-14804 Improve comments on ProcMeshComp

diff --git a/Engine/Plugins/Runtime/ProceduralMeshComponent/Source/ProceduralMeshComponent/Public/KismetProceduralMeshLibrary.h b/Engine/Plugins/Runtime/ProceduralMeshComponent/Source/ProceduralMeshComponent/Public/KismetProceduralMeshLibrary.h
--- a/Engine/Plugins/Runtime/ProceduralMeshComponent/Source/ProceduralMeshComponent/Public/KismetProceduralMeshLibrary.h
+++ b/Engine/Plugins/Runtime/ProceduralMeshComponent/Source/ProceduralMeshComponent/Public/KismetProceduralMeshLibrary.h
@@ -11,15 +11,28 @@ class PROCEDURALMESHCOMPONENT_API UKismetProceduralMeshLibrary : public UBluepri
 {
 	GENERATED_UCLASS_BODY()
 
+	/** Generate vertex and index buffer for a simple box, given the supplied dimensions. Normals, UVs and tangents are also generated for each vertex. */
 	UFUNCTION(BlueprintCallable, Category = "Components|ProceduralMesh")
 	static void GenerateBoxMesh(FVector BoxRadius, TArray<FVector>& Vertices, TArray<int32>& Triangles, TArray<FVector>& Normals, TArray<FVector2D>& UVs, TArray<FProcMeshTangent>& Tangents);
 
+	/** 
+	 *	Automatically generate normals and tangent vectors for a mesh
+	 *	UVs are required for correct tangent generation.
+	 */
 	UFUNCTION(BlueprintCallable, Category = "Components|ProceduralMesh", meta=(AutoCreateRefTerm = "SmoothingGroups,UVs" ))
 	static void CalculateTangentsForMesh(const TArray<FVector>& Vertices, const TArray<int32>& Triangles, const TArray<FVector2D>& UVs, TArray<FVector>& Normals, TArray<FProcMeshTangent>& Tangents);
 
+	/** Add a quad, specified by four indices, to a triangle index buffer as two triangles. */
 	UFUNCTION(BlueprintCallable, Category = "Components|ProceduralMesh")
 	static void ConvertQuadToTriangles(UPARAM(ref) TArray<int32>& Triangles, int32 Vert0, int32 Vert1, int32 Vert2, int32 Vert3);
 
+	/** 
+	 *	Generate an index buffer for a grid of quads. 
+	 *	@param	NumX			Number of vertices in X direction (must be >= 2)
+	 *	@param	NumY			Number of vertices in y direction (must be >= 2)
+	 *	@param	bWinding		Reverses winding of indices generated for each quad
+	 *	@out	Triangles		Output index buffer
+	 */
 	UFUNCTION(BlueprintCallable, Category = "Components|ProceduralMesh")
 	static void CreateGridMeshTriangles(int32 NumX, int32 NumY, bool bWinding, TArray<int32>& Triangles);
 };
diff --git a/Engine/Plugins/Runtime/ProceduralMeshComponent/Source/ProceduralMeshComponent/Public/ProceduralMeshComponent.h b/Engine/Plugins/Runtime/ProceduralMeshComponent/Source/ProceduralMeshComponent/Public/ProceduralMeshComponent.h
--- a/Engine/Plugins/Runtime/ProceduralMeshComponent/Source/ProceduralMeshComponent/Public/ProceduralMeshComponent.h
+++ b/Engine/Plugins/Runtime/ProceduralMeshComponent/Source/ProceduralMeshComponent/Public/ProceduralMeshComponent.h
@@ -4,14 +4,20 @@
 #include "Components/MeshComponent.h"
 #include "ProceduralMeshComponent.generated.h"
 
+/** 
+ *	Struct used to specify a tangent vector for a vertex 
+ *	The Y tangent is computed from the cross product of the vertex normal (Tangent Z) and the TangentX member.
+ */
 USTRUCT(BlueprintType)
 struct FProcMeshTangent
 {
 	GENERATED_USTRUCT_BODY()
 
+	/** Direction of X tangent for this vertex */
 	UPROPERTY(EditAnywhere, BlueprintReadWrite, Category = Tangent)
 	FVector TangentX;
 
+	/** Bool that indicates whether we should flip the Y tangent when we compute it using cross product */
 	UPROPERTY(EditAnywhere, BlueprintReadWrite, Category = Tangent)
 	bool bFlipTangentY;
 
@@ -31,7 +37,7 @@ struct FProcMeshTangent
 	{}
 };
 
-
+/** One vertex for the procedural mesh, used for storing data internally */
 USTRUCT()
 struct FProcMeshVertex
 {
@@ -62,21 +68,22 @@ struct FProcMeshVertex
 	{}
 };
 
+/** One section of the procedural mesh. Each material has its own section. */
 USTRUCT()
 struct FProcMeshSection
 {
 	GENERATED_USTRUCT_BODY()
 
-	/** Vertex buffer for proc mesh */
+	/** Vertex buffer for this section */
 	UPROPERTY()
 	TArray<FProcMeshVertex> ProcVertexBuffer;
-	/** Index buffer for proc mesh */
+	/** Index buffer for this section */
 	UPROPERTY()
 	TArray<int32> ProcIndexBuffer;
-	/** Local bounds of mesh */
+	/** Local bounding box of section */
 	UPROPERTY()
 	FBox SectionLocalBox;
-	/** Should we build collision data for traingles in this section */
+	/** Should we build collision data for triangles in this section */
 	UPROPERTY()
 	bool bEnableCollision;
 
@@ -101,13 +108,25 @@ class PROCEDURALMESHCOMPONENT_API UProceduralMeshComponent : public UMeshCompone
 {
 	GENERATED_UCLASS_BODY()
 
-	/** Add triangle to the geometry to use on this triangle mesh. Must call BuildMesh before geometry can be seen. */
+	/** 
+	 *	Create/replace a section for this procedural mesh component.
+	 *	@param	SectionIndex		Index of the section to create or replace.
+	 *	@param	Vertices			Vertex buffer of all vertex positions to use for this mesh section.
+	 *	@param	Triangles			Index buffer indicating which vertices make up each triangle. Length must be a multiple of 3.
+	 *	@param	Normals				Optional array of normal vectors for each vertex. If supplied, must be same length as Vertices array.
+	 *	@param	UV0					Optional array of texture co-ordinates for each vertex. If supplied, must be same length as Vertices array.
+	 *	@param	VertexColors		Optional array of colors for each vertex. If supplied, must be same length as Vertices array.
+	 *	@param	Tangents			Optional array of tangent vector for each vertex. If supplied, must be same length as Vertices array.
+	 *	@param	bCreateCollision	Indicates whether collision should be created for this section. This adds significant cost.
+	 */
 	UFUNCTION(BlueprintCallable, Category = "Components|ProceduralMesh", meta=(AutoCreateRefTerm = "Normals,UV0,VertexColors,Tangents" ))
 	void CreateMeshSection(int32 SectionIndex, const TArray<FVector>& Vertices, const TArray<int32>& Triangles, const TArray<FVector>& Normals, const TArray<FVector2D>& UV0, const TArray<FColor>& VertexColors, const TArray<FProcMeshTangent>& Tangents, bool bCreateCollision);
 
+	/** Clear a section of the procedural mesh. Other sections do not change index. */
 	UFUNCTION(BlueprintCallable, Category = "Components|ProceduralMesh")
 	void ClearMeshSection(int32 SectionIndex);
 
+	/** Clear all mesh sections and reset to empty state */
 	UFUNCTION(BlueprintCallable, Category = "Components|ProceduralMesh")
 	void ClearAllMeshSections();
 
