UE-15146 : Fix GLEO due to BodySetup in ProceduralMeshComponent

diff --git a/Engine/Plugins/Runtime/ProceduralMeshComponent/Source/ProceduralMeshComponent/Public/ProceduralMeshComponent.h b/Engine/Plugins/Runtime/ProceduralMeshComponent/Source/ProceduralMeshComponent/Public/ProceduralMeshComponent.h
--- a/Engine/Plugins/Runtime/ProceduralMeshComponent/Source/ProceduralMeshComponent/Public/ProceduralMeshComponent.h
+++ b/Engine/Plugins/Runtime/ProceduralMeshComponent/Source/ProceduralMeshComponent/Public/ProceduralMeshComponent.h
@@ -138,7 +138,7 @@ class PROCEDURALMESHCOMPONENT_API UProceduralMeshComponent : public UMeshCompone
 	// End Interface_CollisionDataProvider Interface
 
 	/** Collision data */
-	UPROPERTY()
+	UPROPERTY(transient, duplicatetransient)
 	class UBodySetup* ProcMeshBodySetup;
 
 private:
