UE-15654 : Fix crash in ProcMeshComp when building lighting
Add optimization for creating ProcMeshComp vertex buffer

diff --git a/Engine/Plugins/Runtime/ProceduralMeshComponent/Source/ProceduralMeshComponent/Private/ProceduralMeshComponent.cpp b/Engine/Plugins/Runtime/ProceduralMeshComponent/Source/ProceduralMeshComponent/Private/ProceduralMeshComponent.cpp
--- a/Engine/Plugins/Runtime/ProceduralMeshComponent/Source/ProceduralMeshComponent/Private/ProceduralMeshComponent.cpp
+++ b/Engine/Plugins/Runtime/ProceduralMeshComponent/Source/ProceduralMeshComponent/Private/ProceduralMeshComponent.cpp
@@ -5,6 +5,27 @@
 #include "DynamicMeshBuilder.h"
 
 
+/** Resource array to pass  */
+class FProcMeshVertexResourceArray : public FResourceArrayInterface
+{
+public:
+	FProcMeshVertexResourceArray(void* InData, uint32 InSize)
+		: Data(InData)
+		, Size(InSize)
+	{
+	}
+
+	virtual const void* GetResourceData() const override { return Data; }
+	virtual uint32 GetResourceDataSize() const override { return Size; }
+	virtual void Discard() override { }
+	virtual bool IsStatic() const override { return false; }
+	virtual bool GetAllowCPUAccess() const override { return false; }
+	virtual void SetAllowCPUAccess(bool bInNeedsCPUAccess) override { }
+
+private:
+	void* Data;
+	uint32 Size;
+};
 
 /** Vertex Buffer */
 class FProcMeshVertexBuffer : public FVertexBuffer 
@@ -14,13 +35,13 @@ public:
 
 	virtual void InitRHI() override
 	{
-		FRHIResourceCreateInfo CreateInfo;
-		VertexBufferRHI = RHICreateVertexBuffer(Vertices.Num() * sizeof(FDynamicMeshVertex),BUF_Static,CreateInfo);
+		const uint32 SizeInBytes = Vertices.Num() * sizeof(FDynamicMeshVertex);
 
-		// Copy the vertex data into the vertex buffer.
-		void* VertexBufferData = RHILockVertexBuffer(VertexBufferRHI,0,Vertices.Num() * sizeof(FDynamicMeshVertex), RLM_WriteOnly);
-		FMemory::Memcpy(VertexBufferData,Vertices.GetData(),Vertices.Num() * sizeof(FDynamicMeshVertex));
-		RHIUnlockVertexBuffer(VertexBufferRHI);
+		FProcMeshVertexResourceArray ResourceArray(Vertices.GetData(), SizeInBytes);
+		FRHIResourceCreateInfo CreateInfo(&ResourceArray);
+		VertexBufferRHI = RHICreateVertexBuffer(SizeInBytes, BUF_Static, CreateInfo);
+
+		Vertices.Empty();
 	}
 
 };
@@ -51,28 +72,40 @@ public:
 	FProcMeshVertexFactory()
 	{}
 
+	/** Init function that should only be called on render thread. */
+	void Init_RenderThread(const FProcMeshVertexBuffer* VertexBuffer)
+	{
+		check(IsInRenderingThread());
+
+		// Initialize the vertex factory's stream components.
+		DataType NewData;
+		NewData.PositionComponent = STRUCTMEMBER_VERTEXSTREAMCOMPONENT(VertexBuffer, FDynamicMeshVertex, Position, VET_Float3);
+		NewData.TextureCoordinates.Add(
+			FVertexStreamComponent(VertexBuffer, STRUCT_OFFSET(FDynamicMeshVertex, TextureCoordinate), sizeof(FDynamicMeshVertex), VET_Float2)
+			);
+		NewData.TangentBasisComponents[0] = STRUCTMEMBER_VERTEXSTREAMCOMPONENT(VertexBuffer, FDynamicMeshVertex, TangentX, VET_PackedNormal);
+		NewData.TangentBasisComponents[1] = STRUCTMEMBER_VERTEXSTREAMCOMPONENT(VertexBuffer, FDynamicMeshVertex, TangentZ, VET_PackedNormal);
+		NewData.ColorComponent = STRUCTMEMBER_VERTEXSTREAMCOMPONENT(VertexBuffer, FDynamicMeshVertex, Color, VET_Color);
+		SetData(NewData);
+	}
 
-	/** Initialization */
+	/** Init function that can be called on any thread, and will do the right thing (enqueue command if called on main thread) */
 	void Init(const FProcMeshVertexBuffer* VertexBuffer)
 	{
-		check(!IsInRenderingThread());
-
-		ENQUEUE_UNIQUE_RENDER_COMMAND_TWOPARAMETER(
-			InitProcMeshVertexFactory,
-			FProcMeshVertexFactory*,VertexFactory,this,
-			const FProcMeshVertexBuffer*,VertexBuffer,VertexBuffer,
+		if(IsInRenderingThread())
 		{
-			// Initialize the vertex factory's stream components.
-			DataType NewData;
-			NewData.PositionComponent = STRUCTMEMBER_VERTEXSTREAMCOMPONENT(VertexBuffer,FDynamicMeshVertex,Position,VET_Float3);
-			NewData.TextureCoordinates.Add(
-				FVertexStreamComponent(VertexBuffer,STRUCT_OFFSET(FDynamicMeshVertex,TextureCoordinate),sizeof(FDynamicMeshVertex),VET_Float2)
-				);
-			NewData.TangentBasisComponents[0] = STRUCTMEMBER_VERTEXSTREAMCOMPONENT(VertexBuffer,FDynamicMeshVertex,TangentX,VET_PackedNormal);
-			NewData.TangentBasisComponents[1] = STRUCTMEMBER_VERTEXSTREAMCOMPONENT(VertexBuffer,FDynamicMeshVertex,TangentZ,VET_PackedNormal);
-			NewData.ColorComponent = STRUCTMEMBER_VERTEXSTREAMCOMPONENT(VertexBuffer, FDynamicMeshVertex, Color, VET_Color);
-			VertexFactory->SetData(NewData);
-		});
+			Init_RenderThread(VertexBuffer);
+		}
+		else
+		{
+			ENQUEUE_UNIQUE_RENDER_COMMAND_TWOPARAMETER(
+				InitProcMeshVertexFactory,
+				FProcMeshVertexFactory*,VertexFactory,this,
+				const FProcMeshVertexBuffer*,VertexBuffer,VertexBuffer,
+			{
+				VertexFactory->Init_RenderThread(VertexBuffer);
+			});
+		}
 	}
 };
 
