UE-16041 : Add Experimental flag to ProcMeshComp for 4.8

diff --git a/Engine/Plugins/Runtime/ProceduralMeshComponent/Source/ProceduralMeshComponent/Public/ProceduralMeshComponent.h b/Engine/Plugins/Runtime/ProceduralMeshComponent/Source/ProceduralMeshComponent/Public/ProceduralMeshComponent.h
--- a/Engine/Plugins/Runtime/ProceduralMeshComponent/Source/ProceduralMeshComponent/Public/ProceduralMeshComponent.h
+++ b/Engine/Plugins/Runtime/ProceduralMeshComponent/Source/ProceduralMeshComponent/Public/ProceduralMeshComponent.h
@@ -102,8 +102,11 @@ struct FProcMeshSection
 	}
 };
 
-/** Component that allows you to specify custom triangle mesh geometry */
-UCLASS(hidecategories=(Object,LOD), meta=(BlueprintSpawnableComponent), ClassGroup=Rendering)
+/** 
+ *	Component that allows you to specify custom triangle mesh geometry 
+ *	Beware! This feature is experimental and may be substantially changed in future releases.
+ */
+UCLASS(hidecategories=(Object,LOD), meta=(BlueprintSpawnableComponent), Experimental, ClassGroup=Experimental)
 class PROCEDURALMESHCOMPONENT_API UProceduralMeshComponent : public UMeshComponent, public IInterface_CollisionDataProvider
 {
 	GENERATED_UCLASS_BODY()
