Remove BodySetup.h from engine.h and other unnecessary headers

diff --git a/Engine/Plugins/2D/Paper2D/Source/Paper2D/Private/GroupedSpriteSceneProxy.cpp b/Engine/Plugins/2D/Paper2D/Source/Paper2D/Private/GroupedSpriteSceneProxy.cpp
--- a/Engine/Plugins/2D/Paper2D/Source/Paper2D/Private/GroupedSpriteSceneProxy.cpp
+++ b/Engine/Plugins/2D/Paper2D/Source/Paper2D/Private/GroupedSpriteSceneProxy.cpp
@@ -3,6 +3,7 @@
 #include "Paper2DPrivatePCH.h"
 #include "GroupedSpriteSceneProxy.h"
 #include "PaperGroupedSpriteComponent.h"
+#include "PhysicsEngine/BodySetup.h"
 
 //////////////////////////////////////////////////////////////////////////
 // FGroupedSpriteSceneProxy
diff --git a/Engine/Plugins/Runtime/ProceduralMeshComponent/Source/ProceduralMeshComponent/Private/ProceduralMeshComponent.cpp b/Engine/Plugins/Runtime/ProceduralMeshComponent/Source/ProceduralMeshComponent/Private/ProceduralMeshComponent.cpp
--- a/Engine/Plugins/Runtime/ProceduralMeshComponent/Source/ProceduralMeshComponent/Private/ProceduralMeshComponent.cpp
+++ b/Engine/Plugins/Runtime/ProceduralMeshComponent/Source/ProceduralMeshComponent/Private/ProceduralMeshComponent.cpp
@@ -3,6 +3,7 @@
 #include "ProceduralMeshComponentPluginPrivatePCH.h"
 #include "ProceduralMeshComponent.h"
 #include "DynamicMeshBuilder.h"
+#include "PhysicsEngine/BodySetup.h"
 
 
 /** Resource array to pass  */
diff --git a/Engine/Source/Editor/MeshPaint/Private/MeshPaintStaticMeshAdapter.cpp b/Engine/Source/Editor/MeshPaint/Private/MeshPaintStaticMeshAdapter.cpp
--- a/Engine/Source/Editor/MeshPaint/Private/MeshPaintStaticMeshAdapter.cpp
+++ b/Engine/Source/Editor/MeshPaint/Private/MeshPaintStaticMeshAdapter.cpp
@@ -10,6 +10,7 @@
 // 
 #include "StaticMeshResources.h"
 #include "MeshPaintEdMode.h"
+#include "PhysicsEngine/BodySetup.h"
 // #include "Factories.h"
 // #include "ScopedTransaction.h"
 // #include "MeshPaintRendering.h"
diff --git a/Engine/Source/Editor/PhAT/Private/PhAT.h b/Engine/Source/Editor/PhAT/Private/PhAT.h
--- a/Engine/Source/Editor/PhAT/Private/PhAT.h
+++ b/Engine/Source/Editor/PhAT/Private/PhAT.h
@@ -5,7 +5,7 @@
 #include "Toolkits/AssetEditorToolkit.h"
 #include "PhATSharedData.h"
 #include "EditorUndoClient.h"
-#include "PhysicsEngine/BodySetup.h"
+#include "PhysicsEngine/BodySetupEnums.h"
 
 class SPhATPreviewViewport;
 class FPhATTreeInfo;
diff --git a/Engine/Source/Editor/StaticMeshEditor/Private/StaticMeshEditor.cpp b/Engine/Source/Editor/StaticMeshEditor/Private/StaticMeshEditor.cpp
--- a/Engine/Source/Editor/StaticMeshEditor/Private/StaticMeshEditor.cpp
+++ b/Engine/Source/Editor/StaticMeshEditor/Private/StaticMeshEditor.cpp
@@ -30,6 +30,7 @@
 #include "SNotificationList.h"
 #include "NotificationManager.h"
 #include "Engine/Selection.h"
+#include "PhysicsEngine/BodySetup.h"
 
 #define LOCTEXT_NAMESPACE "StaticMeshEditor"
 
diff --git a/Engine/Source/Editor/StaticMeshEditor/Private/StaticMeshEditorViewportClient.cpp b/Engine/Source/Editor/StaticMeshEditor/Private/StaticMeshEditorViewportClient.cpp
--- a/Engine/Source/Editor/StaticMeshEditor/Private/StaticMeshEditorViewportClient.cpp
+++ b/Engine/Source/Editor/StaticMeshEditor/Private/StaticMeshEditorViewportClient.cpp
@@ -20,6 +20,7 @@
 #include "Runtime/Analytics/Analytics/Public/Interfaces/IAnalyticsProvider.h"
 #include "EngineAnalytics.h"
 #include "AI/Navigation/NavCollision.h"
+#include "PhysicsEngine/BodySetup.h"
 
 #if WITH_PHYSX
 #include "Editor/UnrealEd/Private/EditorPhysXSupport.h"
diff --git a/Engine/Source/Editor/UnrealEd/Private/ConvexDecompTool.cpp b/Engine/Source/Editor/UnrealEd/Private/ConvexDecompTool.cpp
--- a/Engine/Source/Editor/UnrealEd/Private/ConvexDecompTool.cpp
+++ b/Engine/Source/Editor/UnrealEd/Private/ConvexDecompTool.cpp
@@ -8,7 +8,7 @@
 #define USE_VHACD (PLATFORM_WINDOWS || PLATFORM_LINUX || PLATFORM_MAC)
 
 #include "UnrealEd.h"
-
+#include "PhysicsEngine/BodySetup.h"
 
 #if USE_VHACD
 
diff --git a/Engine/Source/Editor/UnrealEd/Private/Fbx/FbxStaticMeshImport.cpp b/Engine/Source/Editor/UnrealEd/Private/Fbx/FbxStaticMeshImport.cpp
--- a/Engine/Source/Editor/UnrealEd/Private/Fbx/FbxStaticMeshImport.cpp
+++ b/Engine/Source/Editor/UnrealEd/Private/Fbx/FbxStaticMeshImport.cpp
@@ -48,6 +48,7 @@
 #include "FbxErrors.h"
 #include "Engine/StaticMeshSocket.h"
 #include "Engine/Polys.h"
+#include "PhysicsEngine/BodySetup.h"
 
 #define LOCTEXT_NAMESPACE "FbxStaticMeshImport"
 
diff --git a/Engine/Source/Editor/UnrealEd/Private/GeomFitUtils.cpp b/Engine/Source/Editor/UnrealEd/Private/GeomFitUtils.cpp
--- a/Engine/Source/Editor/UnrealEd/Private/GeomFitUtils.cpp
+++ b/Engine/Source/Editor/UnrealEd/Private/GeomFitUtils.cpp
@@ -12,6 +12,7 @@
 #include "RawMesh.h"
 #include "MeshUtilities.h"
 #include "Engine/Polys.h"
+#include "PhysicsEngine/BodySetup.h"
 
 #define LOCAL_EPS (0.01f)
 static void AddVertexIfNotPresent(TArray<FVector> &vertices, FVector &newVertex)
diff --git a/Engine/Source/Editor/UnrealEd/Private/StaticMeshEdit.cpp b/Engine/Source/Editor/UnrealEd/Private/StaticMeshEdit.cpp
--- a/Engine/Source/Editor/UnrealEd/Private/StaticMeshEdit.cpp
+++ b/Engine/Source/Editor/UnrealEd/Private/StaticMeshEdit.cpp
@@ -12,6 +12,7 @@
 #include "RawMesh.h"
 #include "MeshUtilities.h"
 #include "Engine/Polys.h"
+#include "PhysicsEngine/BodySetup.h"
 
 bool GBuildStaticMeshCollision = 1;
 
diff --git a/Engine/Source/Runtime/Engine/Public/Engine.h b/Engine/Source/Runtime/Engine/Public/Engine.h
--- a/Engine/Source/Runtime/Engine/Public/Engine.h
+++ b/Engine/Source/Runtime/Engine/Public/Engine.h
@@ -193,7 +193,6 @@
 #include "Kismet/HeadMountedDisplayFunctionLibrary.h"
 #include "Engine/BlueprintGeneratedClass.h"
 #include "Animation/AnimBlueprintGeneratedClass.h"
-#include "PhysicsEngine/BodySetup.h"
 #include "Camera/CameraAnim.h"
 #include "Camera/CameraAnimInst.h"
 #include "Camera/CameraModifier.h"
diff --git a/Engine/Source/Runtime/Engine/Public/StaticMeshResources.h b/Engine/Source/Runtime/Engine/Public/StaticMeshResources.h
--- a/Engine/Source/Runtime/Engine/Public/StaticMeshResources.h
+++ b/Engine/Source/Runtime/Engine/Public/StaticMeshResources.h
@@ -6,13 +6,13 @@
 
 #pragma once
 
-#include "PhysicsEngine/BodySetup.h"
 #include "Engine/StaticMesh.h"
 #include "RawIndexBuffer.h"
 #include "TextureLayout3d.h"
 #include "LocalVertexFactory.h"
 #include "PrimitiveSceneProxy.h"
 #include "SceneManagement.h"
+#include "PhysicsEngine/BodySetupEnums.h"
 
 /**
  * The LOD settings to use for a group of static meshes.
diff --git a/Engine/Source/Runtime/Landscape/Private/LandscapeSplines.cpp b/Engine/Source/Runtime/Landscape/Private/LandscapeSplines.cpp
--- a/Engine/Source/Runtime/Landscape/Private/LandscapeSplines.cpp
+++ b/Engine/Source/Runtime/Landscape/Private/LandscapeSplines.cpp
@@ -18,6 +18,7 @@
 #include "EngineGlobals.h"
 #include "Engine/StaticMesh.h"
 #include "Engine/StaticMeshSocket.h"
+#include "PhysicsEngine/BodySetup.h"
 #if WITH_EDITOR
 #include "LandscapeSplineRaster.h"
 #include "MessageLog.h"
diff --git a/Engine/Source/Runtime/UMG/Private/Components/WidgetComponent.cpp b/Engine/Source/Runtime/UMG/Private/Components/WidgetComponent.cpp
--- a/Engine/Source/Runtime/UMG/Private/Components/WidgetComponent.cpp
+++ b/Engine/Source/Runtime/UMG/Private/Components/WidgetComponent.cpp
@@ -10,6 +10,7 @@
 #include "DynamicMeshBuilder.h"
 #include "Scalability.h"
 #include "WidgetLayoutLibrary.h"
+#include "PhysicsEngine/BodySetup.h"
 
 extern SLATECORE_API int32 bFoldTick;
 
