PR #1411: Fixed procedural mesh components not drawing on PS4 (Contributed by [REDACTED_PERSON])

(Changed slightly to match how its done for other vertexbuffers)

diff --git a/Engine/Plugins/Runtime/ProceduralMeshComponent/Source/Private/ProceduralMeshComponent.cpp b/Engine/Plugins/Runtime/ProceduralMeshComponent/Source/Private/ProceduralMeshComponent.cpp
--- a/Engine/Plugins/Runtime/ProceduralMeshComponent/Source/Private/ProceduralMeshComponent.cpp
+++ b/Engine/Plugins/Runtime/ProceduralMeshComponent/Source/Private/ProceduralMeshComponent.cpp
@@ -39,8 +39,6 @@ public:
 		FProcMeshVertexResourceArray ResourceArray(Vertices.GetData(), SizeInBytes);
 		FRHIResourceCreateInfo CreateInfo(&ResourceArray);
 		VertexBufferRHI = RHICreateVertexBuffer(SizeInBytes, BUF_Static, CreateInfo);
-
-		Vertices.Empty();
 	}
 
 };
