Fixes 2676297's build failure on linux.

diff --git a/Engine/Plugins/Runtime/ProceduralMeshComponent/Source/Private/ProceduralMeshComponent.cpp b/Engine/Plugins/Runtime/ProceduralMeshComponent/Source/Private/ProceduralMeshComponent.cpp
--- a/Engine/Plugins/Runtime/ProceduralMeshComponent/Source/Private/ProceduralMeshComponent.cpp
+++ b/Engine/Plugins/Runtime/ProceduralMeshComponent/Source/Private/ProceduralMeshComponent.cpp
@@ -333,7 +333,7 @@ public:
 #endif
 	}
 
-	virtual FPrimitiveViewRelevance GetViewRelevance(const FSceneView* View)
+	virtual FPrimitiveViewRelevance GetViewRelevance(const FSceneView* View) const
 	{
 		FPrimitiveViewRelevance Result;
 		Result.bDrawRelevance = IsShown(View);
