Fix Win32 alignment compile error

diff --git a/Engine/Plugins/Runtime/ProceduralMeshComponent/Source/Private/KismetProceduralMeshLibrary.cpp b/Engine/Plugins/Runtime/ProceduralMeshComponent/Source/Private/KismetProceduralMeshLibrary.cpp
--- a/Engine/Plugins/Runtime/ProceduralMeshComponent/Source/Private/KismetProceduralMeshLibrary.cpp
+++ b/Engine/Plugins/Runtime/ProceduralMeshComponent/Source/Private/KismetProceduralMeshLibrary.cpp
@@ -456,7 +456,7 @@ void UKismetProceduralMeshLibrary::CopyProceduralMeshFromStaticMeshComponent(USt
 //////////////////////////////////////////////////////////////////////////
 
 /** Util that returns 1 ir on positive side of plane, -1 if negative, or 0 if split by plane */
-int32 BoxPlaneCompare(FBox InBox, FPlane InPlane)
+int32 BoxPlaneCompare(FBox InBox, const FPlane& InPlane)
 {
 	FVector BoxCenter, BoxExtents;
 	InBox.GetCenterAndExtents(BoxCenter, BoxExtents);
