Copying //UE4/Dev-Framework to Dev-Main (//UE4/Dev-Main)

diff --git a/Engine/Config/BaseEditor.ini b/Engine/Config/BaseEditor.ini
--- a/Engine/Config/BaseEditor.ini
+++ b/Engine/Config/BaseEditor.ini
@@ -10,7 +10,7 @@ ShouldLoadLocalizedPropertyNames=True
 [UnrealEd.GeneralSettings]
 EnableStylusPressure=True
 
-[AnimSetViewer]
+[ImportSetting]
 CheckSingleInfluenceLOD=True
 
 [LogWindow]
diff --git a/Engine/Config/BaseEngine.ini b/Engine/Config/BaseEngine.ini
--- a/Engine/Config/BaseEngine.ini
+++ b/Engine/Config/BaseEngine.ini
@@ -462,11 +462,13 @@ bUseStreamingPause=false
 +ActiveClassRedirects=(OldClassName="AnimNode_Root",NewClassName="/Script/AnimGraphRuntime.AnimNode_Root")
 +ActiveClassRedirects=(OldClassName="AnimNode_RotateRootBone",NewClassName="/Script/AnimGraphRuntime.AnimNode_RotateRootBone")
 +ActiveClassRedirects=(OldClassName="AnimNode_RotationOffsetBlendSpace",NewClassName="/Script/AnimGraphRuntime.AnimNode_RotationOffsetBlendSpace")
-+ActiveClassRedirects=(OldClassName="AnimNode_SaveCachedPose",NewClassName="/Script/AnimGraphRuntime.AnimNode_SaveCachedPose")
 +ActiveClassRedirects=(OldClassName="AnimNode_SequenceEvaluator",NewClassName="/Script/AnimGraphRuntime.AnimNode_SequenceEvaluator")
 +ActiveClassRedirects=(OldClassName="AnimNode_Slot",NewClassName="/Script/AnimGraphRuntime.AnimNode_Slot")
 +ActiveClassRedirects=(OldClassName="AnimationNode_TwoWayBlend",NewClassName="/Script/AnimGraphRuntime.AnimationNode_TwoWayBlend")
 
+; Moving cached pose node back to engine
++ActiveClassRedirects=(OldClassName="AnimNode_SaveCachedPose",NewClassName="/Script/Engine.AnimNode_SaveCachedPose")
+
 +TaggedPropertyRedirects=(ClassName="AnimNode_Trail",OldPropertyName="TrailRelaxationCurve",NewPropertyName="TrailRelaxationSpeed")
 
 ;Anim Dynamics redirects
diff --git a/Engine/Plugins/Runtime/ProceduralMeshComponent/ProceduralMeshComponent.uplugin b/Engine/Plugins/Runtime/ProceduralMeshComponent/ProceduralMeshComponent.uplugin
--- a/Engine/Plugins/Runtime/ProceduralMeshComponent/ProceduralMeshComponent.uplugin
+++ b/Engine/Plugins/Runtime/ProceduralMeshComponent/ProceduralMeshComponent.uplugin
@@ -24,6 +24,11 @@
 			[
 				"HTML5"
 			]
+		},
+		{
+			"Name" : "ProceduralMeshComponentEditor",
+			"Type" : "Editor",
+			"LoadingPhase" : "Default"
 		}
 	]
 }
\ No newline at end of file
diff --git a/Engine/Plugins/Runtime/ProceduralMeshComponent/Source/ProceduralMeshComponent/Private/ProceduralMeshComponent.cpp b/Engine/Plugins/Runtime/ProceduralMeshComponent/Source/ProceduralMeshComponent/Private/ProceduralMeshComponent.cpp
new file mode 100644
--- /dev/null
+++ b/Engine/Plugins/Runtime/ProceduralMeshComponent/Source/ProceduralMeshComponent/Private/ProceduralMeshComponent.cpp
@@ -0,0 +1,860 @@
+// Copyright 1998-2016 Epic Games, Inc. All Rights Reserved. 
+
+#include "ProceduralMeshComponentPluginPrivatePCH.h"
+#include "ProceduralMeshComponent.h"
+#include "DynamicMeshBuilder.h"
+
+DECLARE_CYCLE_STAT(TEXT("Create ProcMesh Proxy"), STAT_ProcMesh_CreateSceneProxy, STATGROUP_ProceduralMesh);
+DECLARE_CYCLE_STAT(TEXT("Create Mesh Section"), STAT_ProcMesh_CreateMeshSection, STATGROUP_ProceduralMesh);
+DECLARE_CYCLE_STAT(TEXT("UpdateSection GT"), STAT_ProcMesh_UpdateSectionGT, STATGROUP_ProceduralMesh);
+DECLARE_CYCLE_STAT(TEXT("UpdateSection RT"), STAT_ProcMesh_UpdateSectionRT, STATGROUP_ProceduralMesh);
+DECLARE_CYCLE_STAT(TEXT("Get ProcMesh Elements"), STAT_ProcMesh_GetMeshElements, STATGROUP_ProceduralMesh);
+DECLARE_CYCLE_STAT(TEXT("Update Collision"), STAT_ProcMesh_UpdateCollision, STATGROUP_ProceduralMesh);
+
+
+/** Resource array to pass  */
+class FProcMeshVertexResourceArray : public FResourceArrayInterface
+{
+public:
+	FProcMeshVertexResourceArray(void* InData, uint32 InSize)
+		: Data(InData)
+		, Size(InSize)
+	{
+	}
+
+	virtual const void* GetResourceData() const override { return Data; }
+	virtual uint32 GetResourceDataSize() const override { return Size; }
+	virtual void Discard() override { }
+	virtual bool IsStatic() const override { return false; }
+	virtual bool GetAllowCPUAccess() const override { return false; }
+	virtual void SetAllowCPUAccess(bool bInNeedsCPUAccess) override { }
+
+private:
+	void* Data;
+	uint32 Size;
+};
+
+/** Vertex Buffer */
+class FProcMeshVertexBuffer : public FVertexBuffer
+{
+public:
+	TArray<FDynamicMeshVertex> Vertices;
+
+	virtual void InitRHI() override
+	{
+		const uint32 SizeInBytes = Vertices.Num() * sizeof(FDynamicMeshVertex);
+
+		FProcMeshVertexResourceArray ResourceArray(Vertices.GetData(), SizeInBytes);
+		FRHIResourceCreateInfo CreateInfo(&ResourceArray);
+		VertexBufferRHI = RHICreateVertexBuffer(SizeInBytes, BUF_Static, CreateInfo);
+	}
+
+};
+
+/** Index Buffer */
+class FProcMeshIndexBuffer : public FIndexBuffer
+{
+public:
+	TArray<int32> Indices;
+
+	virtual void InitRHI() override
+	{
+		FRHIResourceCreateInfo CreateInfo;
+		void* Buffer = nullptr;
+		IndexBufferRHI = RHICreateAndLockIndexBuffer(sizeof(int32), Indices.Num() * sizeof(int32), BUF_Static, CreateInfo, Buffer);
+
+		// Write the indices to the index buffer.		
+		FMemory::Memcpy(Buffer, Indices.GetData(), Indices.Num() * sizeof(int32));
+		RHIUnlockIndexBuffer(IndexBufferRHI);
+	}
+};
+
+/** Vertex Factory */
+class FProcMeshVertexFactory : public FLocalVertexFactory
+{
+public:
+
+	FProcMeshVertexFactory()
+	{}
+
+	/** Init function that should only be called on render thread. */
+	void Init_RenderThread(const FProcMeshVertexBuffer* VertexBuffer)
+	{
+		check(IsInRenderingThread());
+
+		// Initialize the vertex factory's stream components.
+		FDataType NewData;
+		NewData.PositionComponent = STRUCTMEMBER_VERTEXSTREAMCOMPONENT(VertexBuffer, FDynamicMeshVertex, Position, VET_Float3);
+		NewData.TextureCoordinates.Add(
+			FVertexStreamComponent(VertexBuffer, STRUCT_OFFSET(FDynamicMeshVertex, TextureCoordinate), sizeof(FDynamicMeshVertex), VET_Float2)
+			);
+		NewData.TangentBasisComponents[0] = STRUCTMEMBER_VERTEXSTREAMCOMPONENT(VertexBuffer, FDynamicMeshVertex, TangentX, VET_PackedNormal);
+		NewData.TangentBasisComponents[1] = STRUCTMEMBER_VERTEXSTREAMCOMPONENT(VertexBuffer, FDynamicMeshVertex, TangentZ, VET_PackedNormal);
+		NewData.ColorComponent = STRUCTMEMBER_VERTEXSTREAMCOMPONENT(VertexBuffer, FDynamicMeshVertex, Color, VET_Color);
+		SetData(NewData);
+	}
+
+	/** Init function that can be called on any thread, and will do the right thing (enqueue command if called on main thread) */
+	void Init(const FProcMeshVertexBuffer* VertexBuffer)
+	{
+		if (IsInRenderingThread())
+		{
+			Init_RenderThread(VertexBuffer);
+		}
+		else
+		{
+			ENQUEUE_UNIQUE_RENDER_COMMAND_TWOPARAMETER(
+				InitProcMeshVertexFactory,
+				FProcMeshVertexFactory*, VertexFactory, this,
+				const FProcMeshVertexBuffer*, VertexBuffer, VertexBuffer,
+				{
+				VertexFactory->Init_RenderThread(VertexBuffer);
+			});
+		}
+	}
+};
+
+/** Class representing a single section of the proc mesh */
+class FProcMeshProxySection
+{
+public:
+	/** Material applied to this section */
+	UMaterialInterface* Material;
+	/** Vertex buffer for this section */
+	FProcMeshVertexBuffer VertexBuffer;
+	/** Index buffer for this section */
+	FProcMeshIndexBuffer IndexBuffer;
+	/** Vertex factory for this section */
+	FProcMeshVertexFactory VertexFactory;
+	/** Whether this section is currently visible */
+	bool bSectionVisible;
+
+	FProcMeshProxySection()
+	: Material(NULL)
+	, bSectionVisible(true)
+	{}
+};
+
+/** 
+ *	Struct used to send update to mesh data 
+ *	Arrays may be empty, in which case no update is performed.
+ */
+class FProcMeshSectionUpdateData
+{
+public:
+	/** Section to update */
+	int32 TargetSection;
+	/** New vertex information */
+	TArray<FProcMeshVertex> NewVertexBuffer;
+};
+
+static void ConvertProcMeshToDynMeshVertex(FDynamicMeshVertex& Vert, const FProcMeshVertex& ProcVert)
+{
+	Vert.Position = ProcVert.Position;
+	Vert.Color = ProcVert.Color;
+	Vert.TextureCoordinate = ProcVert.UV0;
+	Vert.TangentX = ProcVert.Tangent.TangentX;
+	Vert.TangentZ = ProcVert.Normal;
+	Vert.TangentZ.Vector.W = ProcVert.Tangent.bFlipTangentY ? 0 : 255;
+}
+
+/** Procedural mesh scene proxy */
+class FProceduralMeshSceneProxy : public FPrimitiveSceneProxy
+{
+public:
+
+	FProceduralMeshSceneProxy(UProceduralMeshComponent* Component)
+		: FPrimitiveSceneProxy(Component)
+		, BodySetup(Component->GetBodySetup())
+		, MaterialRelevance(Component->GetMaterialRelevance(GetScene().GetFeatureLevel()))
+	{
+		// Copy each section
+		const int32 NumSections = Component->ProcMeshSections.Num();
+		Sections.AddZeroed(NumSections);
+		for (int SectionIdx = 0; SectionIdx < NumSections; SectionIdx++)
+		{
+			FProcMeshSection& SrcSection = Component->ProcMeshSections[SectionIdx];
+			if (SrcSection.ProcIndexBuffer.Num() > 0 && SrcSection.ProcVertexBuffer.Num() > 0)
+			{
+				FProcMeshProxySection* NewSection = new FProcMeshProxySection();
+
+				// Copy data from vertex buffer
+				const int32 NumVerts = SrcSection.ProcVertexBuffer.Num();
+
+				// Allocate verts
+				NewSection->VertexBuffer.Vertices.SetNumUninitialized(NumVerts);
+				// Copy verts
+				for (int VertIdx = 0; VertIdx < NumVerts; VertIdx++)
+				{
+					const FProcMeshVertex& ProcVert = SrcSection.ProcVertexBuffer[VertIdx];
+					FDynamicMeshVertex& Vert = NewSection->VertexBuffer.Vertices[VertIdx];
+					ConvertProcMeshToDynMeshVertex(Vert, ProcVert);
+				}
+
+				// Copy index buffer
+				NewSection->IndexBuffer.Indices = SrcSection.ProcIndexBuffer;
+
+				// Init vertex factory
+				NewSection->VertexFactory.Init(&NewSection->VertexBuffer);
+
+				// Enqueue initialization of render resource
+				BeginInitResource(&NewSection->VertexBuffer);
+				BeginInitResource(&NewSection->IndexBuffer);
+				BeginInitResource(&NewSection->VertexFactory);
+
+				// Grab material
+				NewSection->Material = Component->GetMaterial(SectionIdx);
+				if (NewSection->Material == NULL)
+				{
+					NewSection->Material = UMaterial::GetDefaultMaterial(MD_Surface);
+				}
+
+				// Copy visibility info
+				NewSection->bSectionVisible = SrcSection.bSectionVisible;
+
+				// Save ref to new section
+				Sections[SectionIdx] = NewSection;
+			}
+		}
+	}
+
+	virtual ~FProceduralMeshSceneProxy()
+	{
+		for (FProcMeshProxySection* Section : Sections)
+		{
+			if (Section != nullptr)
+			{
+				Section->VertexBuffer.ReleaseResource();
+				Section->IndexBuffer.ReleaseResource();
+				Section->VertexFactory.ReleaseResource();
+				delete Section;
+			}
+		}
+	}
+
+	/** Called on render thread to assign new dynamic data */
+	void UpdateSection_RenderThread(FProcMeshSectionUpdateData* SectionData)
+	{
+		SCOPE_CYCLE_COUNTER(STAT_ProcMesh_UpdateSectionRT);
+
+		check(IsInRenderingThread());
+
+		// Check we have data 
+		if(	SectionData != nullptr) 			
+		{
+			// Check it references a valid section
+			if (SectionData->TargetSection < Sections.Num() &&
+				Sections[SectionData->TargetSection] != nullptr)
+			{
+				FProcMeshProxySection* Section = Sections[SectionData->TargetSection];
+
+				// Lock vertex buffer
+				const int32 NumVerts = SectionData->NewVertexBuffer.Num();
+				FDynamicMeshVertex* VertexBufferData = (FDynamicMeshVertex*)RHILockVertexBuffer(Section->VertexBuffer.VertexBufferRHI, 0, NumVerts * sizeof(FDynamicMeshVertex), RLM_WriteOnly);
+			
+				// Iterate through vertex data, copying in new info
+				for(int32 VertIdx=0; VertIdx<NumVerts; VertIdx++)
+				{
+					const FProcMeshVertex& ProcVert = SectionData->NewVertexBuffer[VertIdx];
+					FDynamicMeshVertex& Vert = VertexBufferData[VertIdx];
+					ConvertProcMeshToDynMeshVertex(Vert, ProcVert);
+				}
+
+				// Unlock vertex buffer
+				RHIUnlockVertexBuffer(Section->VertexBuffer.VertexBufferRHI);
+			}
+
+			// Free data sent from game thread
+			delete SectionData;
+		}
+	}
+
+	void SetSectionVisibility_RenderThread(int32 SectionIndex, bool bNewVisibility)
+	{
+		check(IsInRenderingThread());
+
+		if(	SectionIndex < Sections.Num() &&
+			Sections[SectionIndex] != nullptr)
+		{
+			Sections[SectionIndex]->bSectionVisible = bNewVisibility;
+		}
+	}
+
+	virtual void GetDynamicMeshElements(const TArray<const FSceneView*>& Views, const FSceneViewFamily& ViewFamily, uint32 VisibilityMap, FMeshElementCollector& Collector) const override
+	{
+		SCOPE_CYCLE_COUNTER(STAT_ProcMesh_GetMeshElements);
+
+
+		// Set up wireframe material (if needed)
+		const bool bWireframe = AllowDebugViewmodes() && ViewFamily.EngineShowFlags.Wireframe;
+
+		FColoredMaterialRenderProxy* WireframeMaterialInstance = NULL;
+		if (bWireframe)
+		{
+			WireframeMaterialInstance = new FColoredMaterialRenderProxy(
+				GEngine->WireframeMaterial ? GEngine->WireframeMaterial->GetRenderProxy(IsSelected()) : NULL,
+				FLinearColor(0, 0.5f, 1.f)
+				);
+
+			Collector.RegisterOneFrameMaterialProxy(WireframeMaterialInstance);
+		}
+
+		// Iterate over sections
+		for (const FProcMeshProxySection* Section : Sections)
+		{
+			if (Section != nullptr && Section->bSectionVisible)
+			{
+				FMaterialRenderProxy* MaterialProxy = bWireframe ? WireframeMaterialInstance : Section->Material->GetRenderProxy(IsSelected());
+
+				// For each view..
+				for (int32 ViewIndex = 0; ViewIndex < Views.Num(); ViewIndex++)
+				{
+					if (VisibilityMap & (1 << ViewIndex))
+					{
+						const FSceneView* View = Views[ViewIndex];
+						// Draw the mesh.
+						FMeshBatch& Mesh = Collector.AllocateMesh();
+						FMeshBatchElement& BatchElement = Mesh.Elements[0];
+						BatchElement.IndexBuffer = &Section->IndexBuffer;
+						Mesh.bWireframe = bWireframe;
+						Mesh.VertexFactory = &Section->VertexFactory;
+						Mesh.MaterialRenderProxy = MaterialProxy;
+						BatchElement.PrimitiveUniformBuffer = CreatePrimitiveUniformBufferImmediate(GetLocalToWorld(), GetBounds(), GetLocalBounds(), true, UseEditorDepthTest());
+						BatchElement.FirstIndex = 0;
+						BatchElement.NumPrimitives = Section->IndexBuffer.Indices.Num() / 3;
+						BatchElement.MinVertexIndex = 0;
+						BatchElement.MaxVertexIndex = Section->VertexBuffer.Vertices.Num() - 1;
+						Mesh.ReverseCulling = IsLocalToWorldDeterminantNegative();
+						Mesh.Type = PT_TriangleList;
+						Mesh.DepthPriorityGroup = SDPG_World;
+						Mesh.bCanApplyViewModeOverrides = false;
+						Collector.AddMesh(ViewIndex, Mesh);
+					}
+				}
+			}
+		}
+
+		// Draw bounds
+#if !(UE_BUILD_SHIPPING || UE_BUILD_TEST)
+		for (int32 ViewIndex = 0; ViewIndex < Views.Num(); ViewIndex++)
+		{
+			if (VisibilityMap & (1 << ViewIndex))
+			{
+				// Draw simple collision as wireframe if 'show collision', and collision is enabled, and we are not using the complex as the simple
+				if (ViewFamily.EngineShowFlags.Collision && IsCollisionEnabled() && BodySetup->GetCollisionTraceFlag() != ECollisionTraceFlag::CTF_UseComplexAsSimple)
+				{
+					FTransform GeomTransform(GetLocalToWorld());
+					BodySetup->AggGeom.GetAggGeom(GeomTransform, GetSelectionColor(FColor(157, 149, 223, 255), IsSelected(), IsHovered()).ToFColor(true), NULL, false, false, UseEditorDepthTest(), ViewIndex, Collector);
+				}
+
+				// Render bounds
+				RenderBounds(Collector.GetPDI(ViewIndex), ViewFamily.EngineShowFlags, GetBounds(), IsSelected());
+			}
+		}
+#endif
+	}
+
+	virtual FPrimitiveViewRelevance GetViewRelevance(const FSceneView* View) const
+	{
+		FPrimitiveViewRelevance Result;
+		Result.bDrawRelevance = IsShown(View);
+		Result.bShadowRelevance = IsShadowCast(View);
+		Result.bDynamicRelevance = true;
+		Result.bRenderInMainPass = ShouldRenderInMainPass();
+		Result.bUsesLightingChannels = GetLightingChannelMask() != GetDefaultLightingChannelMask();
+		Result.bRenderCustomDepth = ShouldRenderCustomDepth();
+		MaterialRelevance.SetPrimitiveViewRelevance(Result);
+		return Result;
+	}
+
+	virtual bool CanBeOccluded() const override
+	{
+		return !MaterialRelevance.bDisableDepthTest;
+	}
+
+	virtual uint32 GetMemoryFootprint(void) const
+	{
+		return(sizeof(*this) + GetAllocatedSize());
+	}
+
+	uint32 GetAllocatedSize(void) const
+	{
+		return(FPrimitiveSceneProxy::GetAllocatedSize());
+	}
+
+private:
+	/** Array of sections */
+	TArray<FProcMeshProxySection*> Sections;
+
+	UBodySetup* BodySetup;
+
+	FMaterialRelevance MaterialRelevance;
+};
+
+//////////////////////////////////////////////////////////////////////////
+
+
+UProceduralMeshComponent::UProceduralMeshComponent(const FObjectInitializer& ObjectInitializer)
+: Super(ObjectInitializer)
+{
+	bUseComplexAsSimpleCollision = true;
+}
+
+void UProceduralMeshComponent::CreateMeshSection_LinearColor(int32 SectionIndex, const TArray<FVector>& Vertices, const TArray<int32>& Triangles, const TArray<FVector>& Normals, const TArray<FVector2D>& UV0, const TArray<FLinearColor>& VertexColors, const TArray<FProcMeshTangent>& Tangents, bool bCreateCollision)
+{
+	// Convert FLinearColors to FColors
+	TArray<FColor> Colors;
+	if (VertexColors.Num() > 0)
+	{
+		Colors.SetNum(VertexColors.Num());
+
+		for (int32 ColorIdx = 0; ColorIdx < VertexColors.Num(); ColorIdx++)
+		{
+			Colors[ColorIdx] = VertexColors[ColorIdx].ToFColor(true);
+		}
+	}
+
+	CreateMeshSection(SectionIndex, Vertices, Triangles, Normals, UV0, Colors, Tangents, bCreateCollision);
+}
+
+void UProceduralMeshComponent::CreateMeshSection(int32 SectionIndex, const TArray<FVector>& Vertices, const TArray<int32>& Triangles, const TArray<FVector>& Normals, const TArray<FVector2D>& UV0, const TArray<FColor>& VertexColors, const TArray<FProcMeshTangent>& Tangents, bool bCreateCollision)
+{
+	SCOPE_CYCLE_COUNTER(STAT_ProcMesh_CreateMeshSection);
+
+	// Ensure sections array is long enough
+	if (SectionIndex >= ProcMeshSections.Num())
+	{
+		ProcMeshSections.SetNum(SectionIndex + 1, false);
+	}
+
+	// Reset this section (in case it already existed)
+	FProcMeshSection& NewSection = ProcMeshSections[SectionIndex];
+	NewSection.Reset();
+
+	// Copy data to vertex buffer
+	const int32 NumVerts = Vertices.Num();
+	NewSection.ProcVertexBuffer.Reset();
+	NewSection.ProcVertexBuffer.AddUninitialized(NumVerts);
+	for (int32 VertIdx = 0; VertIdx < NumVerts; VertIdx++)
+	{
+		FProcMeshVertex& Vertex = NewSection.ProcVertexBuffer[VertIdx];
+
+		Vertex.Position = Vertices[VertIdx];
+		Vertex.Normal = (Normals.Num() == NumVerts) ? Normals[VertIdx] : FVector(0.f, 0.f, 1.f);
+		Vertex.UV0 = (UV0.Num() == NumVerts) ? UV0[VertIdx] : FVector2D(0.f, 0.f);
+		Vertex.Color = (VertexColors.Num() == NumVerts) ? VertexColors[VertIdx] : FColor(255, 255, 255);
+		Vertex.Tangent = (Tangents.Num() == NumVerts) ? Tangents[VertIdx] : FProcMeshTangent();
+
+		// Update bounding box
+		NewSection.SectionLocalBox += Vertex.Position;
+	}
+
+	// Copy index buffer (clamping to vertex range)
+	int32 NumTriIndices = Triangles.Num();
+	NumTriIndices = (NumTriIndices/3) * 3; // Ensure we have exact number of triangles (array is multiple of 3 long)
+
+	NewSection.ProcIndexBuffer.Reset();
+	NewSection.ProcIndexBuffer.AddUninitialized(NumTriIndices);
+	for (int32 IndexIdx = 0; IndexIdx < NumTriIndices; IndexIdx++)
+	{
+		NewSection.ProcIndexBuffer[IndexIdx] = FMath::Min(Triangles[IndexIdx], NumVerts - 1);
+	}
+
+	NewSection.bEnableCollision = bCreateCollision;
+
+	UpdateLocalBounds(); // Update overall bounds
+	UpdateCollision(); // Mark collision as dirty
+	MarkRenderStateDirty(); // New section requires recreating scene proxy
+}
+
+void UProceduralMeshComponent::UpdateMeshSection_LinearColor(int32 SectionIndex, const TArray<FVector>& Vertices, const TArray<FVector>& Normals, const TArray<FVector2D>& UV0, const TArray<FLinearColor>& VertexColors, const TArray<FProcMeshTangent>& Tangents)
+{
+	// Convert FLinearColors to FColors
+	TArray<FColor> Colors;
+	if (VertexColors.Num() > 0)
+	{
+		Colors.SetNum(VertexColors.Num());
+
+		for (int32 ColorIdx = 0; ColorIdx < VertexColors.Num(); ColorIdx++)
+		{
+			Colors[ColorIdx] = VertexColors[ColorIdx].ToFColor(true);
+		}
+	}
+
+	UpdateMeshSection(SectionIndex, Vertices, Normals, UV0, Colors, Tangents);
+}
+
+void UProceduralMeshComponent::UpdateMeshSection(int32 SectionIndex, const TArray<FVector>& Vertices, const TArray<FVector>& Normals, const TArray<FVector2D>& UV0, const TArray<FColor>& VertexColors, const TArray<FProcMeshTangent>& Tangents)
+{
+	SCOPE_CYCLE_COUNTER(STAT_ProcMesh_UpdateSectionGT);
+
+	if(SectionIndex < ProcMeshSections.Num())
+	{
+		FProcMeshSection& Section = ProcMeshSections[SectionIndex];
+		const int32 NumVerts = Section.ProcVertexBuffer.Num();
+
+		// See if positions are changing
+		const bool bPositionsChanging = (Vertices.Num() == NumVerts);
+
+		// Update bounds, if we are getting new position data
+		if (bPositionsChanging)
+		{
+			Section.SectionLocalBox.Init();
+		}
+
+		// Iterate through vertex data, copying in new info
+		for (int32 VertIdx = 0; VertIdx < NumVerts; VertIdx++)
+		{
+			FProcMeshVertex& ModifyVert = Section.ProcVertexBuffer[VertIdx];
+
+			// Position data
+			if (Vertices.Num() == NumVerts)
+			{
+				ModifyVert.Position = Vertices[VertIdx];
+				Section.SectionLocalBox += ModifyVert.Position;
+			}
+
+			// Normal data
+			if (Normals.Num() == NumVerts)
+			{
+				ModifyVert.Normal = Normals[VertIdx];
+			}
+
+			// Tangent data 
+			if (Tangents.Num() == NumVerts)
+			{
+				ModifyVert.Tangent = Tangents[VertIdx];
+			}
+
+			// UV data
+			if (UV0.Num() == NumVerts)
+			{
+				ModifyVert.UV0 = UV0[VertIdx];
+			}
+
+			// Color data
+			if (VertexColors.Num() == NumVerts)
+			{
+				ModifyVert.Color = VertexColors[VertIdx];
+			}
+		}
+
+		if (SceneProxy)
+		{
+			// Create data to update section
+			FProcMeshSectionUpdateData* SectionData = new FProcMeshSectionUpdateData;
+			SectionData->TargetSection = SectionIndex;
+			SectionData->NewVertexBuffer = Section.ProcVertexBuffer;
+
+			// Enqueue command to send to render thread
+			ENQUEUE_UNIQUE_RENDER_COMMAND_TWOPARAMETER(
+				FProcMeshSectionUpdate,
+				FProceduralMeshSceneProxy*, ProcMeshSceneProxy, (FProceduralMeshSceneProxy*)SceneProxy,
+				FProcMeshSectionUpdateData*, SectionData, SectionData,
+				{
+					ProcMeshSceneProxy->UpdateSection_RenderThread(SectionData);
+				}
+			);
+		}
+
+		// If we have collision enabled on this section, update that too
+		if(bPositionsChanging && Section.bEnableCollision)
+		{
+			TArray<FVector> CollisionPositions;
+
+			// We have one collision mesh for all sections, so need to build array of _all_ positions
+			for (const FProcMeshSection& CollisionSection : ProcMeshSections)
+			{
+				// If section has collision, copy it
+				if (CollisionSection.bEnableCollision)
+				{
+					for (int32 VertIdx = 0; VertIdx < CollisionSection.ProcVertexBuffer.Num(); VertIdx++)
+					{
+						CollisionPositions.Add(CollisionSection.ProcVertexBuffer[VertIdx].Position);
+					}
+				}
+			}
+
+			// Pass new positions to trimesh
+			BodyInstance.UpdateTriMeshVertices(CollisionPositions);
+		}
+
+		if (bPositionsChanging)
+		{
+			UpdateLocalBounds(); // Update overall bounds
+			MarkRenderTransformDirty(); // Need to send new bounds to render thread
+		}
+	}
+}
+
+void UProceduralMeshComponent::ClearMeshSection(int32 SectionIndex)
+{
+	if (SectionIndex < ProcMeshSections.Num())
+	{
+		ProcMeshSections[SectionIndex].Reset();
+		UpdateLocalBounds();
+		UpdateCollision();
+		MarkRenderStateDirty();
+	}
+}
+
+void UProceduralMeshComponent::ClearAllMeshSections()
+{
+	ProcMeshSections.Empty();
+	UpdateLocalBounds();
+	UpdateCollision();
+	MarkRenderStateDirty();
+}
+
+void UProceduralMeshComponent::SetMeshSectionVisible(int32 SectionIndex, bool bNewVisibility)
+{
+	if(SectionIndex < ProcMeshSections.Num())
+	{
+		// Set game thread state
+		ProcMeshSections[SectionIndex].bSectionVisible = bNewVisibility;
+
+		if (SceneProxy)
+		{
+			// Enqueue command to modify render thread info
+			ENQUEUE_UNIQUE_RENDER_COMMAND_THREEPARAMETER(
+				FProcMeshSectionVisibilityUpdate,
+				FProceduralMeshSceneProxy*, ProcMeshSceneProxy, (FProceduralMeshSceneProxy*)SceneProxy,
+				int32, SectionIndex, SectionIndex,
+				bool, bNewVisibility, bNewVisibility,
+				{
+					ProcMeshSceneProxy->SetSectionVisibility_RenderThread(SectionIndex, bNewVisibility);
+				}
+			);
+		}
+	}
+}
+
+bool UProceduralMeshComponent::IsMeshSectionVisible(int32 SectionIndex) const
+{
+	return (SectionIndex < ProcMeshSections.Num()) ? ProcMeshSections[SectionIndex].bSectionVisible : false;
+}
+
+int32 UProceduralMeshComponent::GetNumSections() const
+{
+	return ProcMeshSections.Num();
+}
+
+void UProceduralMeshComponent::AddCollisionConvexMesh(TArray<FVector> ConvexVerts)
+{
+	if(ConvexVerts.Num() >= 4)
+	{ 
+		// New element
+		FKConvexElem NewConvexElem;
+		// Copy in vertex info
+		NewConvexElem.VertexData = ConvexVerts;
+		// Update bounding box
+		NewConvexElem.ElemBox = FBox(NewConvexElem.VertexData);
+		// Add to array of convex elements
+		CollisionConvexElems.Add(NewConvexElem);
+		// Refresh collision
+		UpdateCollision();
+	}
+}
+
+void UProceduralMeshComponent::ClearCollisionConvexMeshes()
+{
+	// Empty simple collision info
+	CollisionConvexElems.Empty();
+	// Refresh collision
+	UpdateCollision();
+}
+
+void UProceduralMeshComponent::SetCollisionConvexMeshes(const TArray< TArray<FVector> >& ConvexMeshes)
+{
+	CollisionConvexElems.Reset();
+
+	// Create element for each convex mesh
+	for (int32 ConvexIndex = 0; ConvexIndex < ConvexMeshes.Num(); ConvexIndex++)
+	{
+		FKConvexElem NewConvexElem;
+		NewConvexElem.VertexData = ConvexMeshes[ConvexIndex];
+		NewConvexElem.ElemBox = FBox(NewConvexElem.VertexData);
+
+		CollisionConvexElems.Add(NewConvexElem);
+	}
+
+	UpdateCollision();
+}
+
+
+void UProceduralMeshComponent::UpdateLocalBounds()
+{
+	FBox LocalBox(0);
+
+	for (const FProcMeshSection& Section : ProcMeshSections)
+	{
+		LocalBox += Section.SectionLocalBox;
+	}
+
+	LocalBounds = LocalBox.IsValid ? FBoxSphereBounds(LocalBox) : FBoxSphereBounds(FVector(0, 0, 0), FVector(0, 0, 0), 0); // fallback to reset box sphere bounds
+
+	// Update global bounds
+	UpdateBounds();
+	// Need to send to render thread
+	MarkRenderTransformDirty();
+}
+
+FPrimitiveSceneProxy* UProceduralMeshComponent::CreateSceneProxy()
+{
+	SCOPE_CYCLE_COUNTER(STAT_ProcMesh_CreateSceneProxy);
+
+	return new FProceduralMeshSceneProxy(this);
+}
+
+int32 UProceduralMeshComponent::GetNumMaterials() const
+{
+	return ProcMeshSections.Num();
+}
+
+
+FProcMeshSection* UProceduralMeshComponent::GetProcMeshSection(int32 SectionIndex)
+{
+	if (SectionIndex < ProcMeshSections.Num())
+	{
+		return &ProcMeshSections[SectionIndex];
+	}
+	else
+	{
+		return nullptr;
+	}
+}
+
+
+void UProceduralMeshComponent::SetProcMeshSection(int32 SectionIndex, const FProcMeshSection& Section)
+{
+	// Ensure sections array is long enough
+	if (SectionIndex >= ProcMeshSections.Num())
+	{
+		ProcMeshSections.SetNum(SectionIndex + 1, false);
+	}
+
+	ProcMeshSections[SectionIndex] = Section;
+
+	UpdateLocalBounds(); // Update overall bounds
+	UpdateCollision(); // Mark collision as dirty
+	MarkRenderStateDirty(); // New section requires recreating scene proxy
+}
+
+FBoxSphereBounds UProceduralMeshComponent::CalcBounds(const FTransform& LocalToWorld) const
+{
+	return LocalBounds.TransformBy(LocalToWorld);
+}
+
+bool UProceduralMeshComponent::GetPhysicsTriMeshData(struct FTriMeshCollisionData* CollisionData, bool InUseAllTriData)
+{
+	int32 VertexBase = 0; // Base vertex index for current section
+
+	// For each section..
+	for (int32 SectionIdx = 0; SectionIdx < ProcMeshSections.Num(); SectionIdx++)
+	{
+		FProcMeshSection& Section = ProcMeshSections[SectionIdx];
+		// Do we have collision enabled?
+		if (Section.bEnableCollision)
+		{
+			// Copy vert data
+			for (int32 VertIdx = 0; VertIdx < Section.ProcVertexBuffer.Num(); VertIdx++)
+			{
+				CollisionData->Vertices.Add(Section.ProcVertexBuffer[VertIdx].Position);
+			}
+
+			// Copy triangle data
+			const int32 NumTriangles = Section.ProcIndexBuffer.Num() / 3;
+			for (int32 TriIdx = 0; TriIdx < NumTriangles; TriIdx++)
+			{
+				// Need to add base offset for indices
+				FTriIndices Triangle;
+				Triangle.v0 = Section.ProcIndexBuffer[(TriIdx * 3) + 0] + VertexBase;
+				Triangle.v1 = Section.ProcIndexBuffer[(TriIdx * 3) + 1] + VertexBase;
+				Triangle.v2 = Section.ProcIndexBuffer[(TriIdx * 3) + 2] + VertexBase;
+				CollisionData->Indices.Add(Triangle);
+
+				// Also store matrial info
+				CollisionData->MaterialIndices.Add(SectionIdx);
+			}
+
+			// Remember the base index that new verts will be added from in next section
+			VertexBase = CollisionData->Vertices.Num();
+		}
+	}
+
+	CollisionData->bFlipNormals = true;
+
+	return true;
+}
+
+bool UProceduralMeshComponent::ContainsPhysicsTriMeshData(bool InUseAllTriData) const
+{
+	for (const FProcMeshSection& Section : ProcMeshSections)
+	{
+		if (Section.ProcIndexBuffer.Num() >= 3 && Section.bEnableCollision)
+		{
+			return true;
+		}
+	}
+
+	return false;
+}
+
+void UProceduralMeshComponent::CreateProcMeshBodySetup()
+{
+	if (ProcMeshBodySetup == NULL)
+	{
+		ProcMeshBodySetup = NewObject<UBodySetup>(this);
+		ProcMeshBodySetup->BodySetupGuid = FGuid::NewGuid();
+
+		ProcMeshBodySetup->bGenerateMirroredCollision = false;
+		ProcMeshBodySetup->bDoubleSidedGeometry = true;
+	}
+}
+
+void UProceduralMeshComponent::UpdateCollision()
+{
+	SCOPE_CYCLE_COUNTER(STAT_ProcMesh_UpdateCollision);
+
+	bool bCreatePhysState = false; // Should we create physics state at the end of this function?
+
+	// If its created, shut it down now
+	if (bPhysicsStateCreated)
+	{
+		DestroyPhysicsState();
+		bCreatePhysState = true;
+	}
+
+	// Ensure we have a BodySetup
+	CreateProcMeshBodySetup();
+
+	// Fill in simple collision convex elements
+	ProcMeshBodySetup->AggGeom.ConvexElems = CollisionConvexElems;
+
+	// Set trace flag
+	ProcMeshBodySetup->CollisionTraceFlag = bUseComplexAsSimpleCollision ? CTF_UseComplexAsSimple : CTF_UseDefault;
+
+	// New GUID as collision has changed
+	ProcMeshBodySetup->BodySetupGuid = FGuid::NewGuid();
+
+#if WITH_RUNTIME_PHYSICS_COOKING || WITH_EDITOR
+	// Clear current mesh data
+	ProcMeshBodySetup->InvalidatePhysicsData();
+	// Create new mesh data
+	ProcMeshBodySetup->CreatePhysicsMeshes();
+#endif // WITH_RUNTIME_PHYSICS_COOKING || WITH_EDITOR
+
+	// Create new instance state if desired
+	if (bCreatePhysState)
+	{
+		CreatePhysicsState();
+	}
+}
+
+UBodySetup* UProceduralMeshComponent::GetBodySetup()
+{
+	CreateProcMeshBodySetup();
+	return ProcMeshBodySetup;
+}
+
+
diff --git a/Engine/Plugins/Runtime/ProceduralMeshComponent/Source/ProceduralMeshComponent/Private/ProceduralMeshComponentPluginPrivatePCH.h b/Engine/Plugins/Runtime/ProceduralMeshComponent/Source/ProceduralMeshComponent/Private/ProceduralMeshComponentPluginPrivatePCH.h
new file mode 100644
--- /dev/null
+++ b/Engine/Plugins/Runtime/ProceduralMeshComponent/Source/ProceduralMeshComponent/Private/ProceduralMeshComponentPluginPrivatePCH.h
@@ -0,0 +1,15 @@
+// Copyright 1998-2016 Epic Games, Inc. All Rights Reserved.
+
+#include "CoreUObject.h"
+#include "Engine.h"
+
+// You should place include statements to your module's private header files here.  You only need to
+// add includes for headers that are used in most of your module's source files though.
+#include "IProceduralMeshComponentPlugin.h"
+
+#include "ProceduralMeshComponent.h"
+#include "PhysicsEngine/BodySetup.h"
+
+#pragma once
+
+DECLARE_STATS_GROUP(TEXT("ProceduralMesh"), STATGROUP_ProceduralMesh, STATCAT_Advanced);
diff --git a/Engine/Plugins/Runtime/ProceduralMeshComponent/Source/ProceduralMeshComponent/ProceduralMeshComponent.Build.cs b/Engine/Plugins/Runtime/ProceduralMeshComponent/Source/ProceduralMeshComponent/ProceduralMeshComponent.Build.cs
new file mode 100644
--- /dev/null
+++ b/Engine/Plugins/Runtime/ProceduralMeshComponent/Source/ProceduralMeshComponent/ProceduralMeshComponent.Build.cs
@@ -0,0 +1,25 @@
+// Copyright 1998-2016 Epic Games, Inc. All Rights Reserved.
+
+namespace UnrealBuildTool.Rules
+{
+	public class ProceduralMeshComponent : ModuleRules
+	{
+        public ProceduralMeshComponent(TargetInfo Target)
+		{
+			PrivateIncludePaths.Add("ProceduralMeshComponent/Private");
+            PublicIncludePaths.Add("ProceduralMeshComponent/Public");
+
+			PublicDependencyModuleNames.AddRange(
+				new string[]
+				{
+					"Core",
+					"CoreUObject",
+                    "Engine",
+                    "RenderCore",
+                    "ShaderCore",
+                    "RHI"
+				}
+				);
+		}
+	}
+}
diff --git a/Engine/Plugins/Runtime/ProceduralMeshComponent/Source/ProceduralMeshComponent/Public/KismetProceduralMeshLibrary.h b/Engine/Plugins/Runtime/ProceduralMeshComponent/Source/ProceduralMeshComponent/Public/KismetProceduralMeshLibrary.h
new file mode 100644
--- /dev/null
+++ b/Engine/Plugins/Runtime/ProceduralMeshComponent/Source/ProceduralMeshComponent/Public/KismetProceduralMeshLibrary.h
@@ -0,0 +1,70 @@
+// Copyright 1998-2016 Epic Games, Inc. All Rights Reserved.
+
+#pragma once
+#include "Kismet/BlueprintFunctionLibrary.h"
+#include "ProceduralMeshComponent.h"
+#include "KismetProceduralMeshLibrary.generated.h"
+
+/** Options for creating cap geometry when slicing */
+UENUM()
+enum class EProcMeshSliceCapOption
+{
+	/** Do not create cap geometry */
+	NoCap,
+	/** Add a new section to ProceduralMesh for cap */
+	CreateNewSectionForCap,
+	/** Add cap geometry to existing last section */
+	UseLastSectionForCap
+};
+
+UCLASS()
+class PROCEDURALMESHCOMPONENT_API UKismetProceduralMeshLibrary : public UBlueprintFunctionLibrary
+{
+	GENERATED_UCLASS_BODY()
+
+	/** Generate vertex and index buffer for a simple box, given the supplied dimensions. Normals, UVs and tangents are also generated for each vertex. */
+	UFUNCTION(BlueprintCallable, Category = "Components|ProceduralMesh")
+	static void GenerateBoxMesh(FVector BoxRadius, TArray<FVector>& Vertices, TArray<int32>& Triangles, TArray<FVector>& Normals, TArray<FVector2D>& UVs, TArray<FProcMeshTangent>& Tangents);
+
+	/** 
+	 *	Automatically generate normals and tangent vectors for a mesh
+	 *	UVs are required for correct tangent generation.
+	 */
+	UFUNCTION(BlueprintCallable, Category = "Components|ProceduralMesh", meta=(AutoCreateRefTerm = "SmoothingGroups,UVs" ))
+	static void CalculateTangentsForMesh(const TArray<FVector>& Vertices, const TArray<int32>& Triangles, const TArray<FVector2D>& UVs, TArray<FVector>& Normals, TArray<FProcMeshTangent>& Tangents);
+
+	/** Add a quad, specified by four indices, to a triangle index buffer as two triangles. */
+	UFUNCTION(BlueprintCallable, Category = "Components|ProceduralMesh")
+	static void ConvertQuadToTriangles(UPARAM(ref) TArray<int32>& Triangles, int32 Vert0, int32 Vert1, int32 Vert2, int32 Vert3);
+
+	/** 
+	 *	Generate an index buffer for a grid of quads. 
+	 *	@param	NumX			Number of vertices in X direction (must be >= 2)
+	 *	@param	NumY			Number of vertices in y direction (must be >= 2)
+	 *	@param	bWinding		Reverses winding of indices generated for each quad
+	 *	@out	Triangles		Output index buffer
+	 */
+	UFUNCTION(BlueprintCallable, Category = "Components|ProceduralMesh")
+	static void CreateGridMeshTriangles(int32 NumX, int32 NumY, bool bWinding, TArray<int32>& Triangles);
+
+	/** Grab geometry data from a StaticMesh asset. */
+	UFUNCTION(BlueprintCallable, Category = "Components|ProceduralMesh")
+	static void GetSectionFromStaticMesh(UStaticMesh* InMesh, int32 LODIndex, int32 SectionIndex, TArray<FVector>& Vertices, TArray<int32>& Triangles, TArray<FVector>& Normals, TArray<FVector2D>& UVs, TArray<FProcMeshTangent>& Tangents);
+
+	/** Copy materials from StaticMeshComponent to ProceduralMeshComponent. */
+	UFUNCTION(BlueprintCallable, Category = "Components|ProceduralMesh")
+	static void CopyProceduralMeshFromStaticMeshComponent(UStaticMeshComponent* StaticMeshComponent, int32 LODIndex, UProceduralMeshComponent* ProcMeshComponent, bool bCreateCollision);
+
+	/** 
+	 *	Slice the ProceduralMeshComponent (including simple convex collision) using a plane. Optionally create 'cap' geometry. 
+	 *	@param	InProcMesh				ProceduralMeshComponent to slice
+	 *	@param	PlanePosition			Point on the plane to use for slicing, in world space
+	 *	@param	PlaneNormal				Normal of plane used for slicing. Geometry on the positive side of the plane will be kept.
+	 *	@param	bCreateOtherHalf		If true, an additional ProceduralMeshComponent (OutOtherHalfProcMesh) will be created using the other half of the sliced geometry
+	 *	@param	OutOtherHalfProcMesh	If bCreateOtherHalf is set, this is the new component created. Its owner will be the same as the supplied InProcMesh.
+	 *	@param	CapOption				If and how to create 'cap' geometry on the slicing plane
+	 *	@param	CapMaterial				If creating a new section for the cap, assign this material to that section
+	 */
+	UFUNCTION(BlueprintCallable, Category = "Components|ProceduralMesh")
+	static void SliceProceduralMesh(UProceduralMeshComponent* InProcMesh, FVector PlanePosition, FVector PlaneNormal, bool bCreateOtherHalf, UProceduralMeshComponent*& OutOtherHalfProcMesh, EProcMeshSliceCapOption CapOption, UMaterialInterface* CapMaterial);
+};
diff --git a/Engine/Plugins/Runtime/ProceduralMeshComponent/Source/ProceduralMeshComponentEditor/Private/ProceduralMeshComponentDetails.cpp b/Engine/Plugins/Runtime/ProceduralMeshComponent/Source/ProceduralMeshComponentEditor/Private/ProceduralMeshComponentDetails.cpp
new file mode 100644
--- /dev/null
+++ b/Engine/Plugins/Runtime/ProceduralMeshComponent/Source/ProceduralMeshComponentEditor/Private/ProceduralMeshComponentDetails.cpp
@@ -0,0 +1,174 @@
+// Copyright 1998-2016 Epic Games, Inc. All Rights Reserved.
+
+#include "ProceduralMeshComponentEditorPrivatePCH.h"
+#include "ProceduralMeshComponentDetails.h"
+
+#include "DlgPickAssetPath.h"
+#include "AssetToolsModule.h"
+#include "AssetRegistryModule.h"
+
+#define LOCTEXT_NAMESPACE "ProceduralMeshComponentDetails"
+
+TSharedRef<IDetailCustomization> FProceduralMeshComponentDetails::MakeInstance()
+{
+	return MakeShareable(new FProceduralMeshComponentDetails);
+}
+
+void FProceduralMeshComponentDetails::CustomizeDetails( IDetailLayoutBuilder& DetailBuilder )
+{
+	IDetailCategoryBuilder& ProcMeshCategory = DetailBuilder.EditCategory("ProceduralMesh");
+
+	const FText ConvertToStaticMeshText = LOCTEXT("ConvertToStaticMesh", "Create StaticMesh");
+
+	// Cache set of selected things
+	SelectedObjectsList = DetailBuilder.GetDetailsView().GetSelectedObjects();
+
+	ProcMeshCategory.AddCustomRow(ConvertToStaticMeshText, false)
+	.NameContent()
+	[
+		SNullWidget::NullWidget
+	]
+	.ValueContent()
+	.VAlign(VAlign_Center)
+	.MaxDesiredWidth(250)
+	[
+		SNew(SButton)
+		.VAlign(VAlign_Center)
+		.ToolTipText(LOCTEXT("ConvertToStaticMeshTooltip", "Create a new StaticMesh asset using current geometry from this ProceduralMeshComponent. Does not modify instance."))
+		.OnClicked(this, &FProceduralMeshComponentDetails::ClickedOnConvertToStaticMesh)
+		.Content()
+		[
+			SNew(STextBlock)
+			.Text(ConvertToStaticMeshText)
+		]
+	];
+}
+
+FReply FProceduralMeshComponentDetails::ClickedOnConvertToStaticMesh()
+{
+	// Find first selected ProcMeshComp
+	UProceduralMeshComponent* ProcMeshComp = nullptr;
+	for (TWeakObjectPtr<UObject>& Object : SelectedObjectsList)
+	{
+		ProcMeshComp = Cast<UProceduralMeshComponent>(Object.Get());
+		if (ProcMeshComp != nullptr)
+		{
+			break;
+		}
+	}
+
+	if (ProcMeshComp != nullptr)
+	{
+		FString NewNameSuggestion = FString(TEXT("ProcMesh"));
+		FString PackageName = FString(TEXT("/Game/Meshes/")) + NewNameSuggestion;
+		FString Name;
+		FAssetToolsModule& AssetToolsModule = FModuleManager::LoadModuleChecked<FAssetToolsModule>("AssetTools");
+		AssetToolsModule.Get().CreateUniqueAssetName(PackageName, TEXT(""), PackageName, Name);
+
+		TSharedPtr<SDlgPickAssetPath> PickAssetPathWidget =
+			SNew(SDlgPickAssetPath)
+			.Title(LOCTEXT("ConvertToStaticMeshPickName", "Choose New StaticMesh Location"))
+			.DefaultAssetPath(FText::FromString(PackageName));
+
+		if (PickAssetPathWidget->ShowModal() == EAppReturnType::Ok)
+		{
+			// Get the full name of where we want to create the physics asset.
+			FString UserPackageName = PickAssetPathWidget->GetFullAssetPath().ToString();
+			FName MeshName(*FPackageName::GetLongPackageAssetName(UserPackageName));
+
+			// Check if the user inputed a valid asset name, if they did not, give it the generated default name
+			if (MeshName == NAME_None)
+			{
+				// Use the defaults that were already generated.
+				UserPackageName = PackageName;
+				MeshName = *Name;
+			}
+
+			// Then find/create it.
+			UPackage* Package = CreatePackage(NULL, *UserPackageName);
+			check(Package);
+
+			// Create StaticMesh object
+			UStaticMesh* StaticMesh = NewObject<UStaticMesh>(Package, MeshName, RF_Public | RF_Standalone);
+			StaticMesh->InitResources();
+
+			StaticMesh->LightingGuid = FGuid::NewGuid();
+
+			// Raw mesh data we are filling in
+			FRawMesh RawMesh;
+
+			const int32 NumSections = ProcMeshComp->GetNumSections();
+			int32 VertexBase = 0;
+			for (int32 SectionIdx = 0; SectionIdx < NumSections; SectionIdx++)
+			{
+				FProcMeshSection* ProcSection = ProcMeshComp->GetProcMeshSection(SectionIdx);
+
+				// Copy verts
+				for (FProcMeshVertex& Vert : ProcSection->ProcVertexBuffer)
+				{
+					RawMesh.VertexPositions.Add(Vert.Position);
+				}
+
+				// Copy 'wedge' info
+				int32 NumIndices = ProcSection->ProcIndexBuffer.Num();
+				for (int32 IndexIdx=0; IndexIdx < NumIndices; IndexIdx++)
+				{
+					int32 Index = ProcSection->ProcIndexBuffer[IndexIdx];
+
+					RawMesh.WedgeIndices.Add(Index + VertexBase);
+
+					FProcMeshVertex& ProcVertex = ProcSection->ProcVertexBuffer[Index];
+
+					FVector TangentX = ProcVertex.Tangent.TangentX;
+					FVector TangentZ = ProcVertex.Normal;
+					FVector TangentY = (TangentX ^ TangentZ).GetSafeNormal() * (ProcVertex.Tangent.bFlipTangentY ? -1.f : 1.f);
+
+					RawMesh.WedgeTangentX.Add(TangentX);
+					RawMesh.WedgeTangentY.Add(TangentY);
+					RawMesh.WedgeTangentZ.Add(TangentZ);
+
+					RawMesh.WedgeTexCoords[0].Add(ProcVertex.UV0);
+					RawMesh.WedgeColors.Add(ProcVertex.Color);
+				}
+
+				// copy face info
+				int32 NumTris = NumIndices / 3;
+				for (int32 TriIdx=0; TriIdx < NumTris; TriIdx++)
+				{
+					RawMesh.FaceMaterialIndices.Add(SectionIdx);
+					RawMesh.FaceSmoothingMasks.Add(0); // Assume this is ignored as bRecomputeNormals is false
+				}
+
+				// Copy material
+				StaticMesh->Materials.Add(ProcMeshComp->GetMaterial(SectionIdx));
+
+				// Update offset for creating one big index/vertex buffer
+				VertexBase += ProcSection->ProcVertexBuffer.Num();
+			}
+
+			// Add source to new StaticMesh
+			FStaticMeshSourceModel* SrcModel = new (StaticMesh->SourceModels) FStaticMeshSourceModel();
+			SrcModel->BuildSettings.bRecomputeNormals = false;
+			SrcModel->BuildSettings.bRecomputeTangents = false;
+			SrcModel->BuildSettings.bRemoveDegenerates = false;
+			SrcModel->BuildSettings.bUseHighPrecisionTangentBasis = false;
+			SrcModel->BuildSettings.bUseFullPrecisionUVs = false;
+			SrcModel->BuildSettings.bGenerateLightmapUVs = true;
+			SrcModel->BuildSettings.SrcLightmapIndex = 0;
+			SrcModel->BuildSettings.DstLightmapIndex = 1;
+			SrcModel->RawMeshBulkData->SaveRawMesh(RawMesh);
+
+			// Build mesh from source
+			StaticMesh->Build(false);
+			StaticMesh->PostEditChange();
+
+			// Notify asset registry of new asset
+			FAssetRegistryModule::AssetCreated(StaticMesh);
+		}
+	}
+
+	return FReply::Handled();
+}
+
+
+#undef LOCTEXT_NAMESPACE
diff --git a/Engine/Plugins/Runtime/ProceduralMeshComponent/Source/ProceduralMeshComponentEditor/Private/ProceduralMeshComponentEditorPrivatePCH.h b/Engine/Plugins/Runtime/ProceduralMeshComponent/Source/ProceduralMeshComponentEditor/Private/ProceduralMeshComponentEditorPrivatePCH.h
new file mode 100644
--- /dev/null
+++ b/Engine/Plugins/Runtime/ProceduralMeshComponent/Source/ProceduralMeshComponentEditor/Private/ProceduralMeshComponentEditorPrivatePCH.h
@@ -0,0 +1,17 @@
+// Copyright 1998-2016 Epic Games, Inc. All Rights Reserved.
+
+#include "UnrealEd.h"
+#include "PropertyEditorModule.h"
+#include "DetailLayoutBuilder.h"
+#include "DetailCategoryBuilder.h"
+#include "DetailWidgetRow.h"
+#include "IDetailsView.h"
+#include "PhysicsEngine/ConvexElem.h"
+#include "ProceduralMeshComponent.h"
+#include "Engine/StaticMesh.h"
+#include "RawMesh.h"
+
+#include "IProceduralMeshComponentEditorPlugin.h"
+#include "ProceduralMeshComponentDetails.h"
+
+
diff --git a/Engine/Plugins/Runtime/ProceduralMeshComponent/Source/ProceduralMeshComponentEditor/ProceduralMeshComponentEditor.Build.cs b/Engine/Plugins/Runtime/ProceduralMeshComponent/Source/ProceduralMeshComponentEditor/ProceduralMeshComponentEditor.Build.cs
new file mode 100644
--- /dev/null
+++ b/Engine/Plugins/Runtime/ProceduralMeshComponent/Source/ProceduralMeshComponentEditor/ProceduralMeshComponentEditor.Build.cs
@@ -0,0 +1,33 @@
+// Copyright 1998-2016 Epic Games, Inc. All Rights Reserved.
+
+namespace UnrealBuildTool.Rules
+{
+	public class ProceduralMeshComponentEditor : ModuleRules
+	{
+        public ProceduralMeshComponentEditor(TargetInfo Target)
+		{
+			PrivateIncludePaths.Add("ProceduralMeshComponentEditor/Private");
+            PublicIncludePaths.Add("ProceduralMeshComponentEditor/Public");
+
+			PublicDependencyModuleNames.AddRange(
+				new string[]
+				{
+					"Core",
+					"CoreUObject",
+                    "Slate",
+                    "SlateCore",
+                    "Engine",
+                    "UnrealEd",
+                    "PropertyEditor",
+                    "RenderCore",
+                    "ShaderCore",
+                    "RHI",
+                    "ProceduralMeshComponent",
+                    "RawMesh",
+                    "AssetTools",
+                    "AssetRegistry"
+                }
+				);
+		}
+	}
+}
diff --git a/Engine/Plugins/Runtime/Steam/SteamController/Source/SteamController/Private/SteamController.cpp b/Engine/Plugins/Runtime/Steam/SteamController/Source/SteamController/Private/SteamController.cpp
--- a/Engine/Plugins/Runtime/Steam/SteamController/Source/SteamController/Private/SteamController.cpp
+++ b/Engine/Plugins/Runtime/Steam/SteamController/Source/SteamController/Private/SteamController.cpp
@@ -142,7 +142,7 @@ public:
 		Buttons[22] = FGamepadKeyNames::RightStickLeft;
 		Buttons[23] = FGamepadKeyNames::RightStickRight;
 		Buttons[24] = SteamControllerKeyNames::BackLeft;
-		Buttons[25] = SteamControllerKeyNames::BackLeft;
+		Buttons[25] = SteamControllerKeyNames::BackRight;
 	}
 
 	virtual ~FSteamController()
diff --git a/Engine/Shaders/GpuSkinCacheComputeShader.usf b/Engine/Shaders/GpuSkinCacheComputeShader.usf
--- a/Engine/Shaders/GpuSkinCacheComputeShader.usf
+++ b/Engine/Shaders/GpuSkinCacheComputeShader.usf
@@ -217,8 +217,8 @@ void SkinCacheUpdateBatchCS(uint3 GroupID : SV_GroupID,
 	
 	float3 SPos = SkinPosition(Intermediates) * SkinCacheDebug;
 		
-	float3 TangentX = mul((float3x3)Intermediates.BlendMatrix, LocalTangentX);
-	float3 TangentZ = mul((float3x3)Intermediates.BlendMatrix, LocalTangentZ);
+	float3 TangentX = normalize(mul((float3x3)Intermediates.BlendMatrix, LocalTangentX));
+	float3 TangentZ = normalize(mul((float3x3)Intermediates.BlendMatrix, LocalTangentZ));
 
 	SkinCacheBufferUAV[OutputOffset + GPUSKIN_RWBUFFER_OFFSET_POSITION + 0] = SPos.x;
 	SkinCacheBufferUAV[OutputOffset + GPUSKIN_RWBUFFER_OFFSET_POSITION + 1] = SPos.y;
diff --git a/Engine/Shaders/MorphTargets.usf b/Engine/Shaders/MorphTargets.usf
--- a/Engine/Shaders/MorphTargets.usf
+++ b/Engine/Shaders/MorphTargets.usf
@@ -26,7 +26,7 @@ struct FMorphGPUSkinVertex
 	FVector			DeltaTangentZ;
 };
 
-// This PerVertexInfoList (3 uints)
+// This is how PerVertexInfoList is interpreted as (3 uints)
 struct FPerVertexInfo
 {
 	uint32 DestVertexIndex;
@@ -34,7 +34,7 @@ struct FPerVertexInfo
 	uint32 NumDeltas;
 };
 
-// This FFlattenedDelta (7 uints)
+// This is how FFlattenedDelta is interpreted as (7 uints)
 struct FFlattenedDeltaList
 {
 	FVector PosDelta;
@@ -88,28 +88,27 @@ void ApplyMorphDeltas(uint DestVertexIndex, uint StartDelta, uint NumDeltas, out
 [numthreads(64,1,1)]
 void GPUMorphUpdateCS(uint DispatchId : SV_DispatchThreadID)
 {
-	uint NumLODVertices = (uint)MorphTargetCounts.x;
-
 	// See FMorphDelta
 	const uint MorphDeltaStride = 3 + 3 + 1;
 
-	float3 Position = 0;
-	float3 TangentZ = 0;
-
 	uint NumVertexInfoMorphs = (uint)MorphTargetCounts.y;
 	if (DispatchId.x < NumVertexInfoMorphs)
 	{
 		// See FPerVertexInfo
 		const uint PerVertexInfoStride = 3;
+		// Equivalent to PerVertexInfoList<FPerVertexInfo>
 		const uint DestVertexIndex = PerVertexInfoList[DispatchId.x * PerVertexInfoStride + 0];
 		const uint StartDelta = PerVertexInfoList[DispatchId.x * PerVertexInfoStride + 1];
 		const uint NumDeltas = PerVertexInfoList[DispatchId.x * PerVertexInfoStride + 2];
 
+		float3 Position = 0;
+		float3 TangentZ = 0;
+
 		ApplyMorphDeltas(DestVertexIndex, StartDelta, NumDeltas, Position, TangentZ);
 
 		// See FMorphGPUSkinVertex
 		const uint MorphVertexBufferStride = 3 + 3;
-
+		// Equivalent to MorphVertexBuffer<FMorphGPUSkinVertex>
 		MorphVertexBuffer[MorphVertexBufferStride * DestVertexIndex + 0] = asuint(Position.x);
 		MorphVertexBuffer[MorphVertexBufferStride * DestVertexIndex + 1] = asuint(Position.y);
 		MorphVertexBuffer[MorphVertexBufferStride * DestVertexIndex + 2] = asuint(Position.z);
diff --git a/Engine/Source/Developer/AssetTools/Private/AssetTools.cpp b/Engine/Source/Developer/AssetTools/Private/AssetTools.cpp
--- a/Engine/Source/Developer/AssetTools/Private/AssetTools.cpp
+++ b/Engine/Source/Developer/AssetTools/Private/AssetTools.cpp
@@ -46,6 +46,7 @@ FAssetTools::FAssetTools()
 	RegisterAssetTypeActions( MakeShareable(new FAssetTypeActions_AimOffset) );
 	RegisterAssetTypeActions( MakeShareable(new FAssetTypeActions_AimOffset1D) );
 	RegisterAssetTypeActions( MakeShareable(new FAssetTypeActions_BlendSpace) );
+	RegisterAssetTypeActions( MakeShareable(new FAssetTypeActions_PoseAsset));
 	RegisterAssetTypeActions( MakeShareable(new FAssetTypeActions_BlendSpace1D) );
 	RegisterAssetTypeActions( MakeShareable(new FAssetTypeActions_Blueprint) );
 	RegisterAssetTypeActions( MakeShareable(new FAssetTypeActions_CameraAnim) );
diff --git a/Engine/Source/Developer/AssetTools/Private/AssetToolsPrivatePCH.h b/Engine/Source/Developer/AssetTools/Private/AssetToolsPrivatePCH.h
--- a/Engine/Source/Developer/AssetTools/Private/AssetToolsPrivatePCH.h
+++ b/Engine/Source/Developer/AssetTools/Private/AssetToolsPrivatePCH.h
@@ -69,6 +69,7 @@ DECLARE_LOG_CATEGORY_EXTERN(LogAssetTools, Log, All);
 #include "AssetTypeActions/AssetTypeActions_ParticleSystem.h"
 #include "AssetTypeActions/AssetTypeActions_PhysicalMaterial.h"
 #include "AssetTypeActions/AssetTypeActions_PhysicsAsset.h"
+#include "AssetTypeActions/AssetTypeActions_PoseAsset.h"
 #include "AssetTypeActions/AssetTypeActions_ProceduralFoliageSpawner.h"
 #include "AssetTypeActions/AssetTypeActions_Redirector.h"
 #include "AssetTypeActions/AssetTypeActions_ReverbEffect.h"
diff --git a/Engine/Source/Developer/AssetTools/Private/AssetTypeActions/AssetTypeActions_AnimComposite.h b/Engine/Source/Developer/AssetTools/Private/AssetTypeActions/AssetTypeActions_AnimComposite.h
--- a/Engine/Source/Developer/AssetTools/Private/AssetTypeActions/AssetTypeActions_AnimComposite.h
+++ b/Engine/Source/Developer/AssetTools/Private/AssetTypeActions/AssetTypeActions_AnimComposite.h
@@ -9,7 +9,7 @@ class FAssetTypeActions_AnimComposite : public FAssetTypeActions_AnimationAsset
 public:
 	// IAssetTypeActions Implementation
 	virtual FText GetName() const override { return NSLOCTEXT("AssetTypeActions", "AssetTypeActions_AnimComposite", "Animation Composite"); }
-	virtual FColor GetTypeColor() const override { return FColor(0,128,192); }
+	virtual FColor GetTypeColor() const override { return FColor(181,230,29); }
 	virtual UClass* GetSupportedClass() const override { return UAnimComposite::StaticClass(); }
 	virtual bool CanFilter() override { return true; }
 };
\ No newline at end of file
diff --git a/Engine/Source/Developer/AssetTools/Private/AssetTypeActions/AssetTypeActions_AnimSequence.cpp b/Engine/Source/Developer/AssetTools/Private/AssetTypeActions/AssetTypeActions_AnimSequence.cpp
--- a/Engine/Source/Developer/AssetTools/Private/AssetTypeActions/AssetTypeActions_AnimSequence.cpp
+++ b/Engine/Source/Developer/AssetTools/Private/AssetTypeActions/AssetTypeActions_AnimSequence.cpp
@@ -2,9 +2,16 @@
 
 #include "AssetToolsPrivatePCH.h"
 #include "ContentBrowserModule.h"
+#include "Animation/AnimSequence.h"
+
 
 #define LOCTEXT_NAMESPACE "AssetTypeActions"
 
+UClass* FAssetTypeActions_AnimSequence::GetSupportedClass() const
+{ 
+	return UAnimSequence::StaticClass(); 
+}
+
 void FAssetTypeActions_AnimSequence::GetActions( const TArray<UObject*>& InObjects, FMenuBuilder& MenuBuilder )
 {
 	auto Sequences = GetTypedWeakObjectPtrs<UAnimSequence>(InObjects);
@@ -49,6 +56,16 @@ void FAssetTypeActions_AnimSequence::FillCreateMenu(FMenuBuilder& MenuBuilder, c
 			FCanExecuteAction()
 			)
 		);
+
+	MenuBuilder.AddMenuEntry(
+		LOCTEXT("AnimSequence_NewPoseAsset", "Create PoseAsset"),
+		LOCTEXT("AnimSequence_NewPoseAssetTooltip", "Creates an PoseAsset using the selected anim sequence."),
+		FSlateIcon(FEditorStyle::GetStyleSetName(), "ClassIcon.PoseAsset"),
+		FUIAction(
+		FExecuteAction::CreateSP(this, &FAssetTypeActions_AnimSequence::ExecuteNewPoseAsset, Sequences),
+		FCanExecuteAction()
+		)
+		);
 }
 
 void FAssetTypeActions_AnimSequence::GetResolvedSourceFilePaths(const TArray<UObject*>& TypeAssets, TArray<FString>& OutSourceFilePaths) const
@@ -101,6 +118,13 @@ void FAssetTypeActions_AnimSequence::ExecuteNewAnimMontage(TArray<TWeakObjectPtr
 	CreateAnimationAssets(Objects, UAnimMontage::StaticClass(), Factory, DefaultSuffix, FOnConfigureFactory::CreateSP(this, &FAssetTypeActions_AnimSequence::ConfigureFactoryForAnimMontage));
 }
 
+void FAssetTypeActions_AnimSequence::ExecuteNewPoseAsset(TArray<TWeakObjectPtr<UAnimSequence>> Objects) const
+{
+	const FString DefaultSuffix = TEXT("_PoseAsset");
+	UPoseAssetFactory* Factory = NewObject<UPoseAssetFactory>();
+	CreateAnimationAssets(Objects, UPoseAsset::StaticClass(), Factory, DefaultSuffix, FOnConfigureFactory::CreateSP(this, &FAssetTypeActions_AnimSequence::ConfigureFactoryForPoseAsset));
+}
+
 void FAssetTypeActions_AnimSequence::ConfigureFactoryForAnimComposite(UFactory* AssetFactory, UAnimSequence* SourceAnimation) const
 {
 	UAnimCompositeFactory* CompositeFactory = CastChecked<UAnimCompositeFactory>(AssetFactory);
@@ -113,6 +137,11 @@ void FAssetTypeActions_AnimSequence::ConfigureFactoryForAnimMontage(UFactory* As
 	MontageFactory->SourceAnimation = SourceAnimation;
 }
 
+void FAssetTypeActions_AnimSequence::ConfigureFactoryForPoseAsset(UFactory* AssetFactory, UAnimSequence* SourceAnimation) const
+{
+	UPoseAssetFactory* CompositeFactory = CastChecked<UPoseAssetFactory>(AssetFactory);
+	CompositeFactory->SourceAnimation = SourceAnimation;
+}
 void FAssetTypeActions_AnimSequence::CreateAnimationAssets(const TArray<TWeakObjectPtr<UAnimSequence>>& AnimSequences, TSubclassOf<UAnimationAsset> AssetClass, UFactory* AssetFactory, const FString& InSuffix, FOnConfigureFactory OnConfigureFactory) const
 {
 	if ( AnimSequences.Num() == 1 )
diff --git a/Engine/Source/Developer/AssetTools/Private/AssetTypeActions/AssetTypeActions_AnimSequence.h b/Engine/Source/Developer/AssetTools/Private/AssetTypeActions/AssetTypeActions_AnimSequence.h
--- a/Engine/Source/Developer/AssetTools/Private/AssetTypeActions/AssetTypeActions_AnimSequence.h
+++ b/Engine/Source/Developer/AssetTools/Private/AssetTypeActions/AssetTypeActions_AnimSequence.h
@@ -10,7 +10,7 @@ class FAssetTypeActions_AnimSequence : public FAssetTypeActions_AnimationAsset
 public:
 	// IAssetTypeActions Implementation
 	virtual FText GetName() const override { return NSLOCTEXT("AssetTypeActions", "AssetTypeActions_AnimSequence", "Animation Sequence"); }
-	virtual UClass* GetSupportedClass() const override { return UAnimSequence::StaticClass(); }
+	virtual UClass* GetSupportedClass() const override;
 	virtual void GetActions( const TArray<UObject*>& InObjects, FMenuBuilder& MenuBuilder ) override;
 	virtual bool CanFilter() override { return true; }
 	virtual bool IsImportedAsset() const override { return true; }
@@ -26,12 +26,18 @@ private:
 	/** Handler for when Create AnimMontage is selected */
 	void ExecuteNewAnimMontage(TArray<TWeakObjectPtr<UAnimSequence>> Objects) const;
 
+	/** Handler for when Create Pose Asset is selected */
+	void ExecuteNewPoseAsset(TArray<TWeakObjectPtr<UAnimSequence>> Objects) const;
+
 	/** Delegate handler passed to CreateAnimationAssets when creating AnimComposites */
 	void ConfigureFactoryForAnimComposite(UFactory* AssetFactory, UAnimSequence* SourceAnimation) const;
 
 	/** Delegate handler passed to CreateAnimationAssets when creating AnimMontages */
 	void ConfigureFactoryForAnimMontage(UFactory* AssetFactory, UAnimSequence* SourceAnimation) const;
 
+	/** Delegate handler passed to CreateAnimationAssets when creating PoseAssets */
+	void ConfigureFactoryForPoseAsset(UFactory* AssetFactory, UAnimSequence* SourceAnimation) const;
+
 	/** Creates animation assets of the supplied class */
 	void CreateAnimationAssets(const TArray<TWeakObjectPtr<UAnimSequence>>& AnimSequences, TSubclassOf<UAnimationAsset> AssetClass, UFactory* AssetFactory, const FString& InSuffix, FOnConfigureFactory OnConfigureFactory) const;
 
diff --git a/Engine/Source/Developer/AssetTools/Private/AssetTypeActions/AssetTypeActions_PoseAsset.h b/Engine/Source/Developer/AssetTools/Private/AssetTypeActions/AssetTypeActions_PoseAsset.h
new file mode 100644
--- /dev/null
+++ b/Engine/Source/Developer/AssetTools/Private/AssetTypeActions/AssetTypeActions_PoseAsset.h
@@ -0,0 +1,15 @@
+// Copyright 1998-2016 Epic Games, Inc. All Rights Reserved.
+
+#pragma once
+#include "Animation/PoseAsset.h"
+#include "AssetTypeActions_AnimationAsset.h"
+
+class FAssetTypeActions_PoseAsset : public FAssetTypeActions_AnimationAsset
+{
+public:
+	// IAssetTypeActions Implementation
+	virtual FText GetName() const override { return NSLOCTEXT("AssetTypeActions", "AssetTypeActions_PoseAsset", "Pose Asset"); }
+	virtual FColor GetTypeColor() const override { return FColor(237, 28, 36); }
+	virtual UClass* GetSupportedClass() const override { return UPoseAsset::StaticClass(); }
+	virtual bool CanFilter() override { return true; }
+};
\ No newline at end of file
diff --git a/Engine/Source/Developer/MaterialUtilities/Private/MeshRendering.cpp b/Engine/Source/Developer/MaterialUtilities/Private/MeshRendering.cpp
--- a/Engine/Source/Developer/MaterialUtilities/Private/MeshRendering.cpp
+++ b/Engine/Source/Developer/MaterialUtilities/Private/MeshRendering.cpp
@@ -452,7 +452,7 @@ public:
 			if (Section.MaterialIndex == Data.MaterialIndex)
 			{
 				NumTris += Section.NumTriangles;
-				NumVerts += LODModel.Chunks[Section.ChunkIndex].GetNumVertices();
+				NumVerts += Section.GetNumVertices();
 			}
 		}
 
@@ -503,9 +503,8 @@ public:
 			for (int32 SectionIndex = 0; SectionIndex < SectionCount; SectionIndex++)
 			{
 				const FSkelMeshSection& Section = LODModel.Sections[SectionIndex];
-				const FSkelMeshChunk& Chunk = LODModel.Chunks[Section.ChunkIndex];
 
-				const int32 NumVertsInChunk = Chunk.GetNumVertices();
+				const int32 NumVertsInSection = Section.GetNumVertices();
 
 				if (Section.MaterialIndex == Data.MaterialIndex)
 				{
@@ -514,7 +513,7 @@ public:
 
 					// copy vertices
 					int32 SrcVertIndex = FirstVertex;
-					for (int32 VertIndex = 0; VertIndex < NumVertsInChunk; VertIndex++)
+					for (int32 VertIndex = 0; VertIndex < NumVertsInSection; VertIndex++)
 					{
 						const FSoftSkinVertex& SrcVert = Vertices[SrcVertIndex];
 						FMaterialMeshVertex* DstVert = new(OutVerts)FMaterialMeshVertex();
@@ -551,7 +550,7 @@ public:
 						}
 					}
 				}
-				FirstVertex += NumVertsInChunk;
+				FirstVertex += NumVertsInSection;
 			}
 		}
 		else // bUseNewUVs
diff --git a/Engine/Source/Developer/MeshBoneReduction/Private/MeshBoneReduction.cpp b/Engine/Source/Developer/MeshBoneReduction/Private/MeshBoneReduction.cpp
--- a/Engine/Source/Developer/MeshBoneReduction/Private/MeshBoneReduction.cpp
+++ b/Engine/Source/Developer/MeshBoneReduction/Private/MeshBoneReduction.cpp
@@ -98,26 +98,26 @@ public:
 		return ( OutBonesToReplace.Num() > 0 );
 	}
 
-	void FixUpChunkBoneMaps( FSkelMeshChunk & Chunk, const TMap<FBoneIndexType, FBoneIndexType> &BonesToRepair ) override
+	void FixUpSectionBoneMaps( FSkelMeshSection & Section, const TMap<FBoneIndexType, FBoneIndexType> &BonesToRepair ) override
 	{
 		// now you have list of bones, remove them from vertex influences
 		{
 			TMap<uint8, uint8> BoneMapRemapTable;
 			// first go through bone map and see if this contains BonesToRemove
-			int32 BoneMapSize = Chunk.BoneMap.Num();
+			int32 BoneMapSize = Section.BoneMap.Num();
 			int32 AdjustIndex=0;
 
 			for (int32 BoneMapIndex=0; BoneMapIndex < BoneMapSize; ++BoneMapIndex )
 			{
 				// look for this bone to be removed or not?
-				const FBoneIndexType* ParentBoneIndex = BonesToRepair.Find(Chunk.BoneMap[BoneMapIndex]);
+				const FBoneIndexType* ParentBoneIndex = BonesToRepair.Find(Section.BoneMap[BoneMapIndex]);
 				if ( ParentBoneIndex  )
 				{
 					// this should not happen, I don't ever remove root
 					check (*ParentBoneIndex!=INDEX_NONE);
 
 					// if Parent already exists in the current BoneMap, we just have to fix up the mapping
-					int32 ParentBoneMapIndex = Chunk.BoneMap.Find(*ParentBoneIndex);
+					int32 ParentBoneMapIndex = Section.BoneMap.Find(*ParentBoneIndex);
 
 					// if it exists
 					if (ParentBoneMapIndex != INDEX_NONE)
@@ -128,15 +128,15 @@ public:
 							--ParentBoneMapIndex;
 						}
 
-						// remove current chunk count, will replace with parent
-						Chunk.BoneMap.RemoveAt(BoneMapIndex);
+						// remove current section count, will replace with parent
+						Section.BoneMap.RemoveAt(BoneMapIndex);
 					}
 					else
 					{
 						// if parent doens't exists, we have to add one
 						// this doesn't change bone map size 
-						Chunk.BoneMap.RemoveAt(BoneMapIndex);
-						ParentBoneMapIndex = Chunk.BoneMap.Add(*ParentBoneIndex);
+						Section.BoneMap.RemoveAt(BoneMapIndex);
+						ParentBoneMapIndex = Section.BoneMap.Add(*ParentBoneIndex);
 					}
 
 					// first fix up all indices of BoneMapRemapTable for the indices higher than BoneMapIndex, since BoneMapIndex is being removed
@@ -183,22 +183,10 @@ public:
 
 			if ( BoneMapRemapTable.Num() > 0 )
 			{
-				// fix up rigid verts
-				for (int32 VertIndex=0; VertIndex < Chunk.RigidVertices.Num(); ++VertIndex)
-				{
-					FRigidSkinVertex & Vert = Chunk.RigidVertices[VertIndex];
-
-					uint8 *RemappedBone = BoneMapRemapTable.Find(Vert.Bone);
-					if (RemappedBone)
-					{
-						Vert.Bone = *RemappedBone;
-					}
-				}
-
 				// fix up soft verts
-				for (int32 VertIndex=0; VertIndex < Chunk.SoftVertices.Num(); ++VertIndex)
+				for (int32 VertIndex=0; VertIndex < Section.SoftVertices.Num(); ++VertIndex)
 				{
-					FSoftSkinVertex & Vert = Chunk.SoftVertices[VertIndex];
+					FSoftSkinVertex & Vert = Section.SoftVertices[VertIndex];
 					bool ShouldRenormalize = false;
 
 					for(int32 InfluenceIndex = 0;InfluenceIndex < MAX_TOTAL_INFLUENCES;InfluenceIndex++)
@@ -274,9 +262,9 @@ public:
 		NewModel->RebuildIndexBuffer( &IndexBufferData, &AdjacencyIndexBufferData );
 
 		// fix up chunks
-		for ( int32 ChunkIndex=0; ChunkIndex< NewModel->Chunks.Num(); ++ChunkIndex )	
+		for ( int32 SectionIndex=0; SectionIndex< NewModel->Sections.Num(); ++SectionIndex)
 		{
-			FixUpChunkBoneMaps(NewModel->Chunks[ChunkIndex], BonesToRemove);
+			FixUpSectionBoneMaps(NewModel->Sections[SectionIndex], BonesToRemove);
 		}
 
 		// fix up RequiredBones/ActiveBoneIndices
diff --git a/Engine/Source/Developer/MeshBoneReduction/Public/MeshBoneReduction.h b/Engine/Source/Developer/MeshBoneReduction/Public/MeshBoneReduction.h
--- a/Engine/Source/Developer/MeshBoneReduction/Public/MeshBoneReduction.h
+++ b/Engine/Source/Developer/MeshBoneReduction/Public/MeshBoneReduction.h
@@ -11,15 +11,15 @@ class IMeshBoneReduction
 {
 public:
 	/**
-	 * Fix up chunk bone maps based on list of bones to remove
+	 * Fix up section bone maps based on list of bones to remove
 	 * List of bones to remove should contains <bone index to remove, bone index to replace to >
 	 * 
-	 * @param	Chunk : Chunk reference to fix up
+	 * @param	Section : Section reference to fix up
 	 * @param	BonesToRemove : List of bones to remove with a pair of [bone index, bone to replace]
 	 *
 	 * @return	true if success 
 	 */
-	virtual void FixUpChunkBoneMaps( FSkelMeshChunk & Chunk, const TMap<FBoneIndexType, FBoneIndexType> &BonesToRepair ) = 0;
+	virtual void FixUpSectionBoneMaps( FSkelMeshSection & Section, const TMap<FBoneIndexType, FBoneIndexType> &BonesToRepair ) = 0;
 	/**
 	 * Get Bones To Remove from the Desired LOD
 	 * List of bones to remove should contains <bone index to remove, bone index to replace to >
diff --git a/Engine/Source/Developer/MeshUtilities/Private/MeshUtilities.cpp b/Engine/Source/Developer/MeshUtilities/Private/MeshUtilities.cpp
--- a/Engine/Source/Developer/MeshUtilities/Private/MeshUtilities.cpp
+++ b/Engine/Source/Developer/MeshUtilities/Private/MeshUtilities.cpp
@@ -1195,7 +1195,6 @@ void FMeshUtilities::RechunkSkeletalMeshModels(USkeletalMesh* SrcMesh, int32 Max
 		check(SectionSortOptions.Num() == Chunks.Num());
 
 		BuildSkeletalModelFromChunks(DestModel, RefSkeleton, Chunks, PointToOriginalMap);
-		check(DestModel.Sections.Num() == DestModel.Chunks.Num());
 		check(DestModel.Sections.Num() == SectionSortOptions.Num());
 
 		DestModel.NumTexCoords = SrcModel.NumTexCoords;
@@ -1231,7 +1230,6 @@ void FMeshUtilities::BuildSkeletalModelFromChunks(FStaticLODModel& LODModel, con
 #if WITH_EDITORONLY_DATA
 	// Clear out any data currently held in the LOD model.
 	LODModel.Sections.Empty();
-	LODModel.Chunks.Empty();
 	LODModel.NumVertices = 0;
 	if (LODModel.MultiSizeIndexContainer.IsIndexBufferValid())
 	{
@@ -1245,15 +1243,12 @@ void FMeshUtilities::BuildSkeletalModelFromChunks(FStaticLODModel& LODModel, con
 
 		FSkelMeshSection& Section = *new(LODModel.Sections) FSkelMeshSection();
 		Section.MaterialIndex = SrcChunk->MaterialIndex;
-		Section.ChunkIndex = ChunkIndex;
-
-		FSkelMeshChunk& Chunk = *new(LODModel.Chunks) FSkelMeshChunk();
-		Exchange(Chunk.BoneMap, SrcChunk->BoneMap);
+		Exchange(Section.BoneMap, SrcChunk->BoneMap);
 
 		// Update the active bone indices on the LOD model.
-		for (int32 BoneIndex = 0; BoneIndex < Chunk.BoneMap.Num(); ++BoneIndex)
+		for (int32 BoneIndex = 0; BoneIndex < Section.BoneMap.Num(); ++BoneIndex)
 		{
-			LODModel.ActiveBoneIndices.AddUnique(Chunk.BoneMap[BoneIndex]);
+			LODModel.ActiveBoneIndices.AddUnique(Section.BoneMap[BoneIndex]);
 		}
 	}
 
@@ -1329,14 +1324,12 @@ void FMeshUtilities::BuildSkeletalModelFromChunks(FStaticLODModel& LODModel, con
 	for (int32 SectionIndex = 0; SectionIndex < LODModel.Sections.Num(); SectionIndex++)
 	{
 		FSkelMeshSection& Section = LODModel.Sections[SectionIndex];
-		int32 ChunkIndex = Section.ChunkIndex;
-		FSkelMeshChunk& Chunk = LODModel.Chunks[ChunkIndex];
-		TArray<FSoftSkinBuildVertex>& ChunkVertices = Chunks[ChunkIndex]->Vertices;
+		TArray<FSoftSkinBuildVertex>& ChunkVertices = Chunks[SectionIndex]->Vertices;
 
 		if (IsInGameThread())
 		{
 			// Only update status if in the game thread.  When importing morph targets, this function can run in another thread
-			GWarn->StatusUpdate(ChunkIndex, LODModel.Chunks.Num(), NSLOCTEXT("UnrealEd", "ProcessingChunks", "Processing Chunks"));
+			GWarn->StatusUpdate(SectionIndex, LODModel.Sections.Num(), NSLOCTEXT("UnrealEd", "ProcessingChunks", "Processing Chunks"));
 		}
 
 		CurrentVertexIndex = 0;
@@ -1344,7 +1337,7 @@ void FMeshUtilities::BuildSkeletalModelFromChunks(FStaticLODModel& LODModel, con
 		PrevMaterialIndex = Section.MaterialIndex;
 
 		// Calculate the offset to this chunk's vertices in the vertex buffer.
-		Chunk.BaseVertexIndex = CurrentChunkBaseVertexIndex = LODModel.NumVertices;
+		Section.BaseVertexIndex = CurrentChunkBaseVertexIndex = LODModel.NumVertices;
 
 		// Update the size of the vertex buffer.
 		LODModel.NumVertices += ChunkVertices.Num();
@@ -1352,80 +1345,48 @@ void FMeshUtilities::BuildSkeletalModelFromChunks(FStaticLODModel& LODModel, con
 		// Separate the section's vertices into rigid and soft vertices.
 		TArray<uint32>& ChunkVertexIndexRemap = *new(VertexIndexRemap)TArray<uint32>();
 		ChunkVertexIndexRemap.AddUninitialized(ChunkVertices.Num());
+
 		for (int32 VertexIndex = 0; VertexIndex < ChunkVertices.Num(); VertexIndex++)
 		{
 			const FSoftSkinBuildVertex& SoftVertex = ChunkVertices[VertexIndex];
-			if (SoftVertex.InfluenceWeights[1] == 0)
-			{
-				FRigidSkinVertex RigidVertex;
-				RigidVertex.Position = SoftVertex.Position;
-				RigidVertex.TangentX = SoftVertex.TangentX;
-				RigidVertex.TangentY = SoftVertex.TangentY;
-				RigidVertex.TangentZ = SoftVertex.TangentZ;
-				FMemory::Memcpy(RigidVertex.UVs, SoftVertex.UVs, sizeof(FVector2D)*MAX_TEXCOORDS);
-				RigidVertex.Color = SoftVertex.Color;
-				RigidVertex.Bone = SoftVertex.InfluenceBones[0];
-				// make sure it exists in bone map
-				check(Chunk.BoneMap.IsValidIndex(SoftVertex.InfluenceBones[0]));
-				Chunk.RigidVertices.Add(RigidVertex);
-				ChunkVertexIndexRemap[VertexIndex] = (uint32)(Chunk.BaseVertexIndex + CurrentVertexIndex);
-				CurrentVertexIndex++;
-				// add the index to the original wedge point source of this vertex
-				RawPointIndices.Add(SoftVertex.PointWedgeIdx);
-				// Also remember import index
-				const int32 RawVertIndex = PointToOriginalMap[SoftVertex.PointWedgeIdx];
-				LODModel.MeshToImportVertexMap.Add(RawVertIndex);
-				LODModel.MaxImportVertex = FMath::Max<float>(LODModel.MaxImportVertex, RawVertIndex);
-			}
-		}
-		for (int32 VertexIndex = 0; VertexIndex < ChunkVertices.Num(); VertexIndex++)
-		{
-			const FSoftSkinBuildVertex& SoftVertex = ChunkVertices[VertexIndex];
-			if (SoftVertex.InfluenceWeights[1] > 0)
-			{
-				FSoftSkinVertex NewVertex;
-				NewVertex.Position = SoftVertex.Position;
-				NewVertex.TangentX = SoftVertex.TangentX;
-				NewVertex.TangentY = SoftVertex.TangentY;
-				NewVertex.TangentZ = SoftVertex.TangentZ;
-				FMemory::Memcpy(NewVertex.UVs, SoftVertex.UVs, sizeof(FVector2D)*MAX_TEXCOORDS);
-				NewVertex.Color = SoftVertex.Color;
-				for (int32 i = 0; i < MAX_TOTAL_INFLUENCES; ++i)
+
+			FSoftSkinVertex NewVertex;
+			NewVertex.Position = SoftVertex.Position;
+			NewVertex.TangentX = SoftVertex.TangentX;
+			NewVertex.TangentY = SoftVertex.TangentY;
+			NewVertex.TangentZ = SoftVertex.TangentZ;
+			FMemory::Memcpy(NewVertex.UVs, SoftVertex.UVs, sizeof(FVector2D)*MAX_TEXCOORDS);
+			NewVertex.Color = SoftVertex.Color;
+			for (int32 i = 0; i < MAX_TOTAL_INFLUENCES; ++i)
+			{
+				// it only adds to the bone map if it has weight on it
+				// BoneMap contains only the bones that has influence with weight of >0.f
+				// so here, just make sure it is included before setting the data
+				if (Section.BoneMap.IsValidIndex(SoftVertex.InfluenceBones[i]))
 				{
-					// it only adds to the bone map if it has weight on it
-					// BoneMap contains only the bones that has influence with weight of >0.f
-					// so here, just make sure it is included before setting the data
-					if (Chunk.BoneMap.IsValidIndex(SoftVertex.InfluenceBones[i]))
-					{
-						NewVertex.InfluenceBones[i] = SoftVertex.InfluenceBones[i];
-						NewVertex.InfluenceWeights[i] = SoftVertex.InfluenceWeights[i];
-					}
+					NewVertex.InfluenceBones[i] = SoftVertex.InfluenceBones[i];
+					NewVertex.InfluenceWeights[i] = SoftVertex.InfluenceWeights[i];
 				}
-				Chunk.SoftVertices.Add(NewVertex);
-				ChunkVertexIndexRemap[VertexIndex] = (uint32)(Chunk.BaseVertexIndex + CurrentVertexIndex);
-				CurrentVertexIndex++;
-				// add the index to the original wedge point source of this vertex
-				RawPointIndices.Add(SoftVertex.PointWedgeIdx);
-				// Also remember import index
-				const int32 RawVertIndex = PointToOriginalMap[SoftVertex.PointWedgeIdx];
-				LODModel.MeshToImportVertexMap.Add(RawVertIndex);
-				LODModel.MaxImportVertex = FMath::Max<float>(LODModel.MaxImportVertex, RawVertIndex);
 			}
+			Section.SoftVertices.Add(NewVertex);
+			ChunkVertexIndexRemap[VertexIndex] = (uint32)(Section.BaseVertexIndex + CurrentVertexIndex);
+			CurrentVertexIndex++;
+			// add the index to the original wedge point source of this vertex
+			RawPointIndices.Add(SoftVertex.PointWedgeIdx);
+			// Also remember import index
+			const int32 RawVertIndex = PointToOriginalMap[SoftVertex.PointWedgeIdx];
+			LODModel.MeshToImportVertexMap.Add(RawVertIndex);
+			LODModel.MaxImportVertex = FMath::Max<float>(LODModel.MaxImportVertex, RawVertIndex);
 		}
 
-		// update total num of verts added
-		Chunk.NumRigidVertices = Chunk.RigidVertices.Num();
-		Chunk.NumSoftVertices = Chunk.SoftVertices.Num();
-
 		// update max bone influences
-		Chunk.CalcMaxBoneInfluences();
+		Section.CalcMaxBoneInfluences();
 
 		// Log info about the chunk.
-		UE_LOG(LogSkeletalMesh, Log, TEXT("Chunk %u: %u rigid vertices, %u soft vertices, %u active bones"),
-			ChunkIndex,
-			Chunk.RigidVertices.Num(),
-			Chunk.SoftVertices.Num(),
-			Chunk.BoneMap.Num()
+		UE_LOG(LogSkeletalMesh, Log, TEXT("Section %u: %u vertices, %u active bones"),
+			SectionIndex,
+			Section.GetNumVertices(),
+			Section.BoneMap.Num()
 			);
 	}
 
@@ -1454,7 +1415,7 @@ void FMeshUtilities::BuildSkeletalModelFromChunks(FStaticLODModel& LODModel, con
 		FRawStaticIndexBuffer16or32Interface* IndexBuffer = LODModel.MultiSizeIndexContainer.GetIndexBuffer();
 		Section.BaseIndex = IndexBuffer->Num();
 		const int32 NumIndices = SectionIndices.Num();
-		const TArray<uint32>& SectionVertexIndexRemap = VertexIndexRemap[Section.ChunkIndex];
+		const TArray<uint32>& SectionVertexIndexRemap = VertexIndexRemap[SectionIndex];
 		for (int32 Index = 0; Index < NumIndices; Index++)
 		{
 			uint32 VertexIndex = SectionVertexIndexRemap[SectionIndices[Index]];
@@ -4683,10 +4644,9 @@ bool FMeshUtilities::BuildSkeletalMesh(FStaticLODModel& LODModel, const FReferen
 			bHasBadSections |= (Section.NumTriangles == 0);
 
 			// Log info about the section.
-			UE_LOG(LogSkeletalMesh, Log, TEXT("Section %u: Material=%u, Chunk=%u, %u triangles"),
+			UE_LOG(LogSkeletalMesh, Log, TEXT("Section %u: Material=%u, %u triangles"),
 				SectionIndex,
 				Section.MaterialIndex,
-				Section.ChunkIndex,
 				Section.NumTriangles
 				);
 		}
@@ -5140,10 +5100,9 @@ bool FMeshUtilities::BuildSkeletalMesh_Legacy(FStaticLODModel& LODModel, const F
 			bHasBadSections |= (Section.NumTriangles == 0);
 
 			// Log info about the section.
-			UE_LOG(LogSkeletalMesh, Log, TEXT("Section %u: Material=%u, Chunk=%u, %u triangles"),
+			UE_LOG(LogSkeletalMesh, Log, TEXT("Section %u: Material=%u, %u triangles"),
 				SectionIndex,
 				Section.MaterialIndex,
-				Section.ChunkIndex,
 				Section.NumTriangles
 				);
 		}
@@ -6421,7 +6380,9 @@ void FMeshUtilities::CalculateTextureCoordinateBoundsForSkeletalMesh(const FStat
 		const int32 MaterialIndex = Section.MaterialIndex;
 
 		if (OutBounds.Num() <= MaterialIndex)
+		{
 			OutBounds.SetNumZeroed(MaterialIndex + 1);
+		}
 
 		for (uint32 Index = FirstIndex; Index < LastIndex; ++Index)
 		{
diff --git a/Engine/Source/Developer/MeshUtilities/Private/SkeletalMeshTools.cpp b/Engine/Source/Developer/MeshUtilities/Private/SkeletalMeshTools.cpp
--- a/Engine/Source/Developer/MeshUtilities/Private/SkeletalMeshTools.cpp
+++ b/Engine/Source/Developer/MeshUtilities/Private/SkeletalMeshTools.cpp
@@ -350,10 +350,10 @@ namespace SkeletalMeshTools
 		}
 		OutData.MeshToImportVertexMap = Model.MeshToImportVertexMap;
 		OutData.Sections = Model.Sections;
-		for (int32 ChunkIndex = 0; ChunkIndex < Model.Chunks.Num(); ++ChunkIndex)
+		for (int32 SectionIndex = 0; SectionIndex < Model.Sections.Num(); ++SectionIndex)
 		{
 			TArray<FBoneIndexType>& DestBoneMap = *new(OutData.BoneMaps) TArray<FBoneIndexType>();
-			DestBoneMap = Model.Chunks[ChunkIndex].BoneMap;
+			DestBoneMap = Model.Sections[SectionIndex].BoneMap;
 		}
 		OutData.NumTexCoords = Model.NumTexCoords;
 	#endif // #if WITH_EDITORONLY_DATA
@@ -459,28 +459,16 @@ namespace SkeletalMeshTools
 		Infos.AddZeroed( SkeletalMesh->RefSkeleton.GetNum() );
 
 		FStaticLODModel* LODModel = &ImportedResource->LODModels[0];
-		for(int32 ChunkIndex = 0;ChunkIndex < LODModel->Chunks.Num();ChunkIndex++)
+		for(int32 SectionIndex = 0; SectionIndex < LODModel->Sections.Num(); SectionIndex++)
 		{
-			FSkelMeshChunk& Chunk = LODModel->Chunks[ChunkIndex];
-			for(int32 i=0; i<Chunk.RigidVertices.Num(); i++)
+			FSkelMeshSection& Section = LODModel->Sections[SectionIndex];
+			for(int32 i=0; i<Section.SoftVertices.Num(); i++)
 			{
-				FRigidSkinVertex* RigidVert = &Chunk.RigidVertices[i];
-				int32 BoneIndex = Chunk.BoneMap[RigidVert->Bone];
-
-				FVector LocalPos = SkeletalMesh->RefBasesInvMatrix[BoneIndex].TransformPosition(RigidVert->Position);
-				Infos[BoneIndex].Positions.Add(LocalPos);
-
-				FVector LocalNormal = SkeletalMesh->RefBasesInvMatrix[BoneIndex].TransformVector(RigidVert->TangentZ);
-				Infos[BoneIndex].Normals.Add(LocalNormal);
-			}
-
-			for(int32 i=0; i<Chunk.SoftVertices.Num(); i++)
-			{
-				FSoftSkinVertex* SoftVert = &Chunk.SoftVertices[i];
+				FSoftSkinVertex* SoftVert = &Section.SoftVertices[i];
 
 				if(bOnlyDominant)
 				{
-					int32 BoneIndex = Chunk.BoneMap[GetDominantBoneIndex(SoftVert)];
+					int32 BoneIndex = Section.BoneMap[GetDominantBoneIndex(SoftVert)];
 
 					FVector LocalPos = SkeletalMesh->RefBasesInvMatrix[BoneIndex].TransformPosition(SoftVert->Position);
 					Infos[BoneIndex].Positions.Add(LocalPos);
@@ -494,7 +482,7 @@ namespace SkeletalMeshTools
 					{
 						if(SoftVert->InfluenceWeights[j] > 0)
 						{
-							int32 BoneIndex = Chunk.BoneMap[SoftVert->InfluenceBones[j]];
+							int32 BoneIndex = Section.BoneMap[SoftVert->InfluenceBones[j]];
 
 							FVector LocalPos = SkeletalMesh->RefBasesInvMatrix[BoneIndex].TransformPosition(SoftVert->Position);
 							Infos[BoneIndex].Positions.Add(LocalPos);
diff --git a/Engine/Source/Developer/SimplygonMeshReduction/Private/SimplygonMeshReduction.cpp b/Engine/Source/Developer/SimplygonMeshReduction/Private/SimplygonMeshReduction.cpp
--- a/Engine/Source/Developer/SimplygonMeshReduction/Private/SimplygonMeshReduction.cpp
+++ b/Engine/Source/Developer/SimplygonMeshReduction/Private/SimplygonMeshReduction.cpp
@@ -346,9 +346,9 @@ public:
 			//todo: check - memory leak here - SrcModel is not released?
 
 			// fix up chunks to remove the bones that set to be removed
-			for (int32 ChunkIndex = 0; ChunkIndex < NewSrcModel->Chunks.Num(); ++ChunkIndex)
+			for (int32 SectionIndex = 0; SectionIndex < NewSrcModel->Sections.Num(); ++SectionIndex)
 			{
-				MeshBoneReductionInterface->FixUpChunkBoneMaps(NewSrcModel->Chunks[ChunkIndex], BonesToRemove);
+				MeshBoneReductionInterface->FixUpSectionBoneMaps(NewSrcModel->Sections[SectionIndex], BonesToRemove);
 			}
 		}
 
@@ -1529,17 +1529,12 @@ private:
 		GeometryData->AddMaterialIds();
 		SimplygonSDK::spRidArray MaterialIndices = GeometryData->GetMaterialIds();
 
-		// Add per-vertex data. This data needs to be added per-chunk so that we can properly map bones.
-#if WITH_APEX_CLOTHING
-		const int32 ChunkCount = LODModel.NumNonClothingSections();
-#else
-		const int32 ChunkCount = LODModel.Chunks.Num();
-#endif // #if WITH_APEX_CLOTHING
+		// Add per-vertex data. This data needs to be added per-section so that we can properly map bones.
 		uint32 FirstVertex = 0;
-		for ( int32 ChunkIndex = 0; ChunkIndex < ChunkCount; ++ChunkIndex )
+		for ( uint32 SectionIndex = 0; SectionIndex < SectionCount; ++SectionIndex )
 		{
-			const FSkelMeshChunk& Chunk = LODModel.Chunks[ ChunkIndex ];
-			const uint32 LastVertex = FirstVertex + (uint32)Chunk.RigidVertices.Num() + (uint32)Chunk.SoftVertices.Num();
+			const FSkelMeshSection& Section = LODModel.Sections[SectionIndex];
+			const uint32 LastVertex = FirstVertex + (uint32)Section.SoftVertices.Num();
 			for ( uint32 VertexIndex = FirstVertex; VertexIndex < LastVertex; ++VertexIndex )
 			{
 				FSoftSkinVertex& Vertex = Vertices[ VertexIndex ];
@@ -1555,8 +1550,8 @@ private:
 
 					if ( BoneInfluence > 0 )
 					{
-						check( BoneIndex < Chunk.BoneMap.Num() );
-						uint32 BoneID = BoneIDs[ Chunk.BoneMap[ BoneIndex ] ];
+						check( BoneIndex < Section.BoneMap.Num() );
+						uint32 BoneID = BoneIDs[Section.BoneMap[ BoneIndex ] ];
 						VertexBoneIds[InfluenceIndex] = BoneID;
 						VertexBoneWeights[InfluenceIndex] = BoneInfluence / 255.0f;
 					}
diff --git a/Engine/Source/Editor/AIGraph/Private/AIGraph.cpp b/Engine/Source/Editor/AIGraph/Private/AIGraph.cpp
--- a/Engine/Source/Editor/AIGraph/Private/AIGraph.cpp
+++ b/Engine/Source/Editor/AIGraph/Private/AIGraph.cpp
@@ -51,15 +51,18 @@ bool UAIGraph::UpdateUnknownNodeClasses()
 	for (int32 NodeIdx = 0; NodeIdx < Nodes.Num(); NodeIdx++)
 	{
 		UAIGraphNode* MyNode = Cast<UAIGraphNode>(Nodes[NodeIdx]);
-		const bool bUpdatedNode = MyNode->RefreshNodeClass();
-		bUpdated = bUpdated || bUpdatedNode;
-
-		for (int32 SubNodeIdx = 0; SubNodeIdx < MyNode->SubNodes.Num(); SubNodeIdx++)
+		if (MyNode)
 		{
-			if (MyNode->SubNodes[SubNodeIdx])
+			const bool bUpdatedNode = MyNode->RefreshNodeClass();
+			bUpdated = bUpdated || bUpdatedNode;
+
+			for (int32 SubNodeIdx = 0; SubNodeIdx < MyNode->SubNodes.Num(); SubNodeIdx++)
 			{
-				const bool bUpdatedSubNode = MyNode->SubNodes[SubNodeIdx]->RefreshNodeClass();
-				bUpdated = bUpdated || bUpdatedSubNode;
+				if (MyNode->SubNodes[SubNodeIdx])
+				{
+					const bool bUpdatedSubNode = MyNode->SubNodes[SubNodeIdx]->RefreshNodeClass();
+					bUpdated = bUpdated || bUpdatedSubNode;
+				}
 			}
 		}
 	}
diff --git a/Engine/Source/Editor/AnimGraph/Classes/AnimGraphNode_AssetPlayerBase.cpp b/Engine/Source/Editor/AnimGraph/Classes/AnimGraphNode_AssetPlayerBase.cpp
--- a/Engine/Source/Editor/AnimGraph/Classes/AnimGraphNode_AssetPlayerBase.cpp
+++ b/Engine/Source/Editor/AnimGraph/Classes/AnimGraphNode_AssetPlayerBase.cpp
@@ -8,6 +8,7 @@
 #include "AnimGraphNode_SequencePlayer.h"
 #include "AnimGraphNode_RotationOffsetBlendSpace.h"
 #include "AnimGraphNode_BlendSpacePlayer.h"
+#include "AnimGraphNode_PoseBlendNode.h"
 
 bool IsAimOffsetBlendSpace(const UClass* BlendSpaceClass)
 {
@@ -36,6 +37,10 @@ UClass* GetNodeClassForAsset(const UClass* AssetClass)
 	{
 		return UAnimGraphNode_SequencePlayer::StaticClass();
 	}
+	else if (AssetClass->IsChildOf(UPoseAsset::StaticClass()))
+	{
+		return UAnimGraphNode_PoseBlendNode::StaticClass();
+	}
 	return nullptr;
 }
 
diff --git a/Engine/Source/Editor/AnimGraph/Classes/AnimGraphNode_PoseBlendNode.h b/Engine/Source/Editor/AnimGraph/Classes/AnimGraphNode_PoseBlendNode.h
new file mode 100644
--- /dev/null
+++ b/Engine/Source/Editor/AnimGraph/Classes/AnimGraphNode_PoseBlendNode.h
@@ -0,0 +1,41 @@
+// Copyright 1998-2016 Epic Games, Inc. All Rights Reserved.
+
+#pragma once
+#include "AnimGraphNode_Base.h"
+#include "AnimGraphNode_AssetPlayerBase.h"
+#include "AnimNodes/AnimNode_PoseBlendNode.h"
+#include "EdGraph/EdGraphNodeUtils.h" // for FNodeTextCache
+#include "AnimGraphNode_PoseBlendNode.generated.h"
+
+UCLASS(MinimalAPI)
+class UAnimGraphNode_PoseBlendNode : public UAnimGraphNode_AssetPlayerBase
+{
+	GENERATED_UCLASS_BODY()
+
+	UPROPERTY(EditAnywhere, Category=Settings)
+	FAnimNode_PoseBlendNode Node;
+
+	// UEdGraphNode interface
+	virtual FText GetNodeTitle(ENodeTitleType::Type TitleType) const override;
+	virtual FText GetTooltipText() const override;
+	// End of UEdGraphNode
+
+	// UAnimGraphNode_Base interface
+	virtual void ValidateAnimNodeDuringCompilation(class USkeleton* ForSkeleton, class FCompilerResultsLog& MessageLog) override;
+	virtual void PreloadRequiredAssets() override;
+	// Interface to support transition getter
+	virtual bool DoesSupportTimeForTransitionGetter() const override;
+	virtual UAnimationAsset* GetAnimationAsset() const override;
+	virtual void GetAllAnimationSequencesReferred(TArray<UAnimationAsset*>& AnimationAssets) const override;
+	virtual void ReplaceReferredAnimations(const TMap<UAnimationAsset*, UAnimationAsset*>& AnimAssetReplacementMap) override;
+	// End of UAnimGraphNode_Base
+
+	// UK2Node interface
+	virtual void GetContextMenuActions(const FGraphNodeContextMenuBuilder& Context) const override;
+	// End of UK2Node interface
+
+	virtual void SetAnimationAsset(UAnimationAsset* Asset) override;
+private:
+	/** Constructing FText strings can be costly, so we cache the node's title */
+	FNodeTextCache CachedNodeTitle;
+};
diff --git a/Engine/Source/Editor/AnimGraph/Classes/AnimGraphNode_SaveCachedPose.h b/Engine/Source/Editor/AnimGraph/Classes/AnimGraphNode_SaveCachedPose.h
--- a/Engine/Source/Editor/AnimGraph/Classes/AnimGraphNode_SaveCachedPose.h
+++ b/Engine/Source/Editor/AnimGraph/Classes/AnimGraphNode_SaveCachedPose.h
@@ -2,7 +2,7 @@
 
 #pragma once
 #include "AnimGraphNode_Base.h"
-#include "AnimNodes/AnimNode_SaveCachedPose.h"
+#include "Animation/AnimNode_SaveCachedPose.h"
 #include "EdGraph/EdGraphNodeUtils.h" // for FNodeTextCache
 #include "AnimGraphNode_SaveCachedPose.generated.h"
 
diff --git a/Engine/Source/Editor/AnimGraph/Private/AnimGraphNode_AnimDynamics.cpp b/Engine/Source/Editor/AnimGraph/Private/AnimGraphNode_AnimDynamics.cpp
--- a/Engine/Source/Editor/AnimGraph/Private/AnimGraphNode_AnimDynamics.cpp
+++ b/Engine/Source/Editor/AnimGraph/Private/AnimGraphNode_AnimDynamics.cpp
@@ -9,6 +9,7 @@
 #include "DetailWidgetRow.h"
 #include "IDetailPropertyRow.h"
 #include "AnimationCustomVersion.h"
+#include "Animation/AnimInstance.h"
 
 #define LOCTEXT_NAMESPACE "AnimDynamicsNode"
 
diff --git a/Engine/Source/Editor/AnimGraph/Private/AnimGraphNode_Base.cpp b/Engine/Source/Editor/AnimGraph/Private/AnimGraphNode_Base.cpp
--- a/Engine/Source/Editor/AnimGraph/Private/AnimGraphNode_Base.cpp
+++ b/Engine/Source/Editor/AnimGraph/Private/AnimGraphNode_Base.cpp
@@ -6,6 +6,7 @@
 #include "AnimationGraphSchema.h"
 #include "BlueprintNodeSpawner.h"
 #include "BlueprintActionDatabaseRegistrar.h"
+#include "Animation/AnimInstance.h"
 
 /////////////////////////////////////////////////////
 // FA3NodeOptionalPinManager
diff --git a/Engine/Source/Editor/AnimGraph/Private/AnimGraphNode_PoseBlendNode.cpp b/Engine/Source/Editor/AnimGraph/Private/AnimGraphNode_PoseBlendNode.cpp
new file mode 100644
--- /dev/null
+++ b/Engine/Source/Editor/AnimGraph/Private/AnimGraphNode_PoseBlendNode.cpp
@@ -0,0 +1,135 @@
+// Copyright 1998-2016 Epic Games, Inc. All Rights Reserved.
+
+#include "AnimGraphPrivatePCH.h"
+
+#include "CompilerResultsLog.h"
+#include "GraphEditorActions.h"
+#include "AnimGraphNode_PoseBlendNode.h"
+
+/////////////////////////////////////////////////////
+// UAnimGraphNode_PoseBlendNode
+
+#define LOCTEXT_NAMESPACE "PoseBlendNode"
+
+UAnimGraphNode_PoseBlendNode::UAnimGraphNode_PoseBlendNode(const FObjectInitializer& ObjectInitializer)
+	: Super(ObjectInitializer)
+{
+}
+
+void UAnimGraphNode_PoseBlendNode::PreloadRequiredAssets()
+{
+	PreloadObject(Node.PoseAsset);
+
+	Super::PreloadRequiredAssets();
+}
+
+
+void UAnimGraphNode_PoseBlendNode::GetAllAnimationSequencesReferred(TArray<UAnimationAsset*>& AnimationAssets) const
+{
+	if(Node.PoseAsset)
+	{
+		HandleAnimReferenceCollection(Node.PoseAsset, AnimationAssets);
+	}
+}
+
+void UAnimGraphNode_PoseBlendNode::ReplaceReferredAnimations(const TMap<UAnimationAsset*, UAnimationAsset*>& AnimAssetReplacementMap)
+{
+	HandleAnimReferenceReplacement(Node.PoseAsset, AnimAssetReplacementMap);
+}
+
+FText UAnimGraphNode_PoseBlendNode::GetTooltipText() const
+{
+	// FText::Format() is slow, so we utilize the cached list title
+	return GetNodeTitle(ENodeTitleType::ListView);
+}
+
+FText UAnimGraphNode_PoseBlendNode::GetNodeTitle(ENodeTitleType::Type TitleType) const
+{
+	if (Node.PoseAsset == nullptr)
+	{
+		return LOCTEXT("PoseByName_TitleNONE", "Pose (None)");
+	}
+	// @TODO: don't know enough about this node type to comfortably assert that
+	//        the CacheName won't change after the node has spawned... until
+	//        then, we'll leave this optimization off
+	else //if (CachedNodeTitle.IsOutOfDate(this))
+	{
+		FFormatNamedArguments Args;
+		Args.Add(TEXT("PoseAssetName"), FText::FromString(Node.PoseAsset->GetName()));
+
+		// FText::Format() is slow, so we cache this to save on performance
+		CachedNodeTitle.SetCachedText(FText::Format(LOCTEXT("PoseByName_Title", "{PoseAssetName}"), Args), this);
+	}
+
+	return CachedNodeTitle;
+}
+
+// void UAnimGraphNode_PoseBlendNode::GetMenuActions(FBlueprintActionDatabaseRegistrar& ActionRegistrar) const
+// {
+// 	// Intentionally empty; you can drop down a regular sequence player and convert into a sequence evaluator in the right-click menu.
+// }
+
+void UAnimGraphNode_PoseBlendNode::SetAnimationAsset(UAnimationAsset* Asset)
+{
+	if (UPoseAsset* PoseAsset =  Cast<UPoseAsset>(Asset))
+	{
+		Node.PoseAsset = PoseAsset;
+	}
+}
+
+void UAnimGraphNode_PoseBlendNode::ValidateAnimNodeDuringCompilation(class USkeleton* ForSkeleton, class FCompilerResultsLog& MessageLog)
+{
+	if (Node.PoseAsset == NULL)
+	{
+		MessageLog.Error(TEXT("@@ references an unknown poseasset"), this);
+	}
+	else
+	{
+		USkeleton* SeqSkeleton = Node.PoseAsset->GetSkeleton();
+		if (SeqSkeleton&& // if anim sequence doesn't have skeleton, it might be due to anim sequence not loaded yet, @todo: wait with anim blueprint compilation until all assets are loaded?
+			!SeqSkeleton->IsCompatible(ForSkeleton))
+		{
+			MessageLog.Error(TEXT("@@ references poseasset that uses different skeleton @@"), this, SeqSkeleton);
+		}
+	}
+}
+
+// potentially in the future, give options for the pose name?
+// void UAnimGraphNode_PoseBlendNode::GetContextMenuActions(const FGraphNodeContextMenuBuilder& Context) const
+// {
+// 	if (!Context.bIsDebugging)
+// 	{
+// 		// add an option to convert to a regular sequence player
+// 		Context.MenuBuilder->BeginSection("AnimGraphNodePoseByName", NSLOCTEXT("A3Nodes", "PoseByNameHeading", "Sequence Evaluator"));
+// 		{
+// 			Context.MenuBuilder->AddMenuEntry(FGraphEditorCommands::Get().OpenRelatedAsset);
+// 			Context.MenuBuilder->AddMenuEntry(FGraphEditorCommands::Get().ConvertToSeqPlayer);
+// 		}
+// 		Context.MenuBuilder->EndSection();
+// 	}
+// }
+
+bool UAnimGraphNode_PoseBlendNode::DoesSupportTimeForTransitionGetter() const
+{
+	return false;
+}
+
+UAnimationAsset* UAnimGraphNode_PoseBlendNode::GetAnimationAsset() const 
+{
+	return Node.PoseAsset;
+}
+
+void UAnimGraphNode_PoseBlendNode::GetContextMenuActions(const FGraphNodeContextMenuBuilder& Context) const
+{
+	if (!Context.bIsDebugging)
+	{
+		// add an option to convert to single frame
+		Context.MenuBuilder->BeginSection("AnimGraphNodePoseBlender", LOCTEXT("PoseBlenderHeading", "Pose Blender"));
+		{
+			Context.MenuBuilder->AddMenuEntry(FGraphEditorCommands::Get().ConvertToPoseByName);
+		}
+		Context.MenuBuilder->EndSection();
+	}
+}
+
+#undef LOCTEXT_NAMESPACE
\ No newline at end of file
diff --git a/Engine/Source/Editor/AnimGraph/Private/AnimGraphNode_PoseByName.cpp b/Engine/Source/Editor/AnimGraph/Private/AnimGraphNode_PoseByName.cpp
new file mode 100644
--- /dev/null
+++ b/Engine/Source/Editor/AnimGraph/Private/AnimGraphNode_PoseByName.cpp
@@ -0,0 +1,135 @@
+// Copyright 1998-2016 Epic Games, Inc. All Rights Reserved.
+
+#include "AnimGraphPrivatePCH.h"
+
+#include "CompilerResultsLog.h"
+#include "GraphEditorActions.h"
+#include "AnimGraphNode_PoseByName.h"
+
+/////////////////////////////////////////////////////
+// UAnimGraphNode_PoseByName
+
+#define LOCTEXT_NAMESPACE "A3Nodes"
+
+UAnimGraphNode_PoseByName::UAnimGraphNode_PoseByName(const FObjectInitializer& ObjectInitializer)
+	: Super(ObjectInitializer)
+{
+}
+
+void UAnimGraphNode_PoseByName::PreloadRequiredAssets()
+{
+	PreloadObject(Node.PoseAsset);
+
+	Super::PreloadRequiredAssets();
+}
+
+
+void UAnimGraphNode_PoseByName::GetAllAnimationSequencesReferred(TArray<UAnimationAsset*>& AnimationAssets) const
+{
+	if(Node.PoseAsset)
+	{
+		HandleAnimReferenceCollection(Node.PoseAsset, AnimationAssets);
+	}
+}
+
+void UAnimGraphNode_PoseByName::ReplaceReferredAnimations(const TMap<UAnimationAsset*, UAnimationAsset*>& AnimAssetReplacementMap)
+{
+	HandleAnimReferenceReplacement(Node.PoseAsset, AnimAssetReplacementMap);
+}
+
+FText UAnimGraphNode_PoseByName::GetTooltipText() const
+{
+	// FText::Format() is slow, so we utilize the cached list title
+	return GetNodeTitle(ENodeTitleType::ListView);
+}
+
+FText UAnimGraphNode_PoseByName::GetNodeTitle(ENodeTitleType::Type TitleType) const
+{
+	if (Node.PoseAsset == nullptr)
+	{
+		return LOCTEXT("PoseByName_TitleNONE", "Pose (None)");
+	}
+	// @TODO: don't know enough about this node type to comfortably assert that
+	//        the CacheName won't change after the node has spawned... until
+	//        then, we'll leave this optimization off
+	else //if (CachedNodeTitle.IsOutOfDate(this))
+	{
+		FFormatNamedArguments Args;
+		Args.Add(TEXT("PoseAssetName"), FText::FromString(Node.PoseAsset->GetName()));
+		Args.Add(TEXT("PoseName"), FText::FromString(Node.PoseName.ToString()));
+
+		// FText::Format() is slow, so we cache this to save on performance
+		CachedNodeTitle.SetCachedText(FText::Format(LOCTEXT("PoseByName_Title", "{PoseAssetName} : {PoseName}"), Args), this);
+	}
+
+	return CachedNodeTitle;
+}
+
+// void UAnimGraphNode_PoseByName::GetMenuActions(FBlueprintActionDatabaseRegistrar& ActionRegistrar) const
+// {
+// 	// Intentionally empty; you can drop down a regular sequence player and convert into a sequence evaluator in the right-click menu.
+// }
+
+void UAnimGraphNode_PoseByName::SetAnimationAsset(UAnimationAsset* Asset)
+{
+	if (UPoseAsset* PoseAsset =  Cast<UPoseAsset>(Asset))
+	{
+		Node.PoseAsset = PoseAsset;
+	}
+}
+
+void UAnimGraphNode_PoseByName::ValidateAnimNodeDuringCompilation(class USkeleton* ForSkeleton, class FCompilerResultsLog& MessageLog)
+{
+	if (Node.PoseAsset == NULL)
+	{
+		MessageLog.Error(TEXT("@@ references an unknown sequence"), this);
+	}
+	else
+	{
+		USkeleton* SeqSkeleton = Node.PoseAsset->GetSkeleton();
+		if (SeqSkeleton&& // if anim sequence doesn't have skeleton, it might be due to anim sequence not loaded yet, @todo: wait with anim blueprint compilation until all assets are loaded?
+			!SeqSkeleton->IsCompatible(ForSkeleton))
+		{
+			MessageLog.Error(TEXT("@@ references sequence that uses different skeleton @@"), this, SeqSkeleton);
+		}
+	}
+}
+
+// potentially in the future, give options for the pose name?
+// void UAnimGraphNode_PoseByName::GetContextMenuActions(const FGraphNodeContextMenuBuilder& Context) const
+// {
+// 	if (!Context.bIsDebugging)
+// 	{
+// 		// add an option to convert to a regular sequence player
+// 		Context.MenuBuilder->BeginSection("AnimGraphNodePoseByName", NSLOCTEXT("A3Nodes", "PoseByNameHeading", "Sequence Evaluator"));
+// 		{
+// 			Context.MenuBuilder->AddMenuEntry(FGraphEditorCommands::Get().OpenRelatedAsset);
+// 			Context.MenuBuilder->AddMenuEntry(FGraphEditorCommands::Get().ConvertToSeqPlayer);
+// 		}
+// 		Context.MenuBuilder->EndSection();
+// 	}
+// }
+
+bool UAnimGraphNode_PoseByName::DoesSupportTimeForTransitionGetter() const
+{
+	return false;
+}
+
+UAnimationAsset* UAnimGraphNode_PoseByName::GetAnimationAsset() const 
+{
+	return Node.PoseAsset;
+}
+
+void UAnimGraphNode_PoseByName::GetContextMenuActions(const FGraphNodeContextMenuBuilder& Context) const
+{
+	if (!Context.bIsDebugging)
+	{
+		// add an option to convert to single frame
+		Context.MenuBuilder->BeginSection("AnimGraphNodePoseByName", LOCTEXT("PoseByNameHeading", "Pose By Name"));
+		{
+			Context.MenuBuilder->AddMenuEntry(FGraphEditorCommands::Get().ConvertToPoseBlender);
+		}
+		Context.MenuBuilder->EndSection();
+	}
+}
+#undef LOCTEXT_NAMESPACE
\ No newline at end of file
diff --git a/Engine/Source/Editor/AnimGraph/Private/AnimGraphNode_WheelHandler.cpp b/Engine/Source/Editor/AnimGraph/Private/AnimGraphNode_WheelHandler.cpp
--- a/Engine/Source/Editor/AnimGraph/Private/AnimGraphNode_WheelHandler.cpp
+++ b/Engine/Source/Editor/AnimGraph/Private/AnimGraphNode_WheelHandler.cpp
@@ -1,6 +1,7 @@
 // Copyright 1998-2016 Epic Games, Inc. All Rights Reserved.
 
 #include "AnimGraphPrivatePCH.h"
+#include "Vehicles/VehicleAnimInstance.h"
 #include "AnimGraphNode_WheelHandler.h"
 #include "CompilerResultsLog.h"
 
