Procedural Mesh performance regressed between 4.20 and 4.21
#fix Rename and reorder operation when updating a mesh section with new vertex data
#misc Make sure we log an error in case a user tries to resize a procedural mesh component section

diff --git a/Engine/Plugins/Runtime/ProceduralMeshComponent/Source/ProceduralMeshComponent/Private/ProceduralMeshComponent.cpp b/Engine/Plugins/Runtime/ProceduralMeshComponent/Source/ProceduralMeshComponent/Private/ProceduralMeshComponent.cpp
--- a/Engine/Plugins/Runtime/ProceduralMeshComponent/Source/ProceduralMeshComponent/Private/ProceduralMeshComponent.cpp
+++ b/Engine/Plugins/Runtime/ProceduralMeshComponent/Source/ProceduralMeshComponent/Private/ProceduralMeshComponent.cpp
@@ -26,6 +26,7 @@ DECLARE_CYCLE_STAT(TEXT("UpdateSection RT"), STAT_ProcMesh_UpdateSectionRT, STAT
 DECLARE_CYCLE_STAT(TEXT("Get ProcMesh Elements"), STAT_ProcMesh_GetMeshElements, STATGROUP_ProceduralMesh);
 DECLARE_CYCLE_STAT(TEXT("Update Collision"), STAT_ProcMesh_UpdateCollision, STATGROUP_ProceduralMesh);
 
+DEFINE_LOG_CATEGORY_STATIC(LogProceduralComponent, Log, All);
 
 /** Resource array to pass  */
 class FProcMeshVertexResourceArray : public FResourceArrayInterface
@@ -497,116 +498,110 @@ void UProceduralMeshComponent::UpdateMeshSection(int32 SectionIndex, const TArra
 	if(SectionIndex < ProcMeshSections.Num())
 	{
 		FProcMeshSection& Section = ProcMeshSections[SectionIndex];
-		const int32 NumVerts = Section.ProcVertexBuffer.Num();
+		const int32 NumVerts = Vertices.Num();
+		const int32 PreviousNumVerts = Section.ProcVertexBuffer.Num();
 
 		// See if positions are changing
-		const bool bPositionsChanging = (Vertices.Num() == NumVerts);
+		const bool bSameVertexCount = PreviousNumVerts == NumVerts;
 
 		// Update bounds, if we are getting new position data
-		if (bPositionsChanging)
+		if (bSameVertexCount)
 		{
-			Section.SectionLocalBox.Init();
-		}
-
-		// Iterate through vertex data, copying in new info
-		for (int32 VertIdx = 0; VertIdx < NumVerts; VertIdx++)
-		{
-			FProcMeshVertex& ModifyVert = Section.ProcVertexBuffer[VertIdx];
+			Section.SectionLocalBox = FBox(Vertices);
 
-			// Position data
-			if (Vertices.Num() == NumVerts)
+			// Iterate through vertex data, copying in new info
+			for (int32 VertIdx = 0; VertIdx < NumVerts; VertIdx++)
 			{
-				ModifyVert.Position = Vertices[VertIdx];
-				Section.SectionLocalBox += ModifyVert.Position;
-			}
-
-			// Normal data
-			if (Normals.Num() == NumVerts)
-			{
-				ModifyVert.Normal = Normals[VertIdx];
-			}
+				FProcMeshVertex& ModifyVert = Section.ProcVertexBuffer[VertIdx];
 
-			// Tangent data 
-			if (Tangents.Num() == NumVerts)
-			{
-				ModifyVert.Tangent = Tangents[VertIdx];
-			}
-
-			// UV0 data
-			if (UV0.Num() == NumVerts)
-			{
-				ModifyVert.UV0 = UV0[VertIdx];
-			}
-			// UV1 data
-			if (UV1.Num() == NumVerts)
-			{
-				ModifyVert.UV1 = UV1[VertIdx];
-			}
-			// UV2 data
-			if (UV2.Num() == NumVerts)
-			{
-				ModifyVert.UV2 = UV2[VertIdx];
-			}
-			// UV3 data
-			if (UV3.Num() == NumVerts)
-			{
-				ModifyVert.UV3 = UV3[VertIdx];
-			}
+				// Position data
+				if (Vertices.Num() == NumVerts)
+				{
+					ModifyVert.Position = Vertices[VertIdx];
+				}
 
-			// Color data
-			if (VertexColors.Num() == NumVerts)
-			{
-				ModifyVert.Color = VertexColors[VertIdx];
-			}
-		}
+				// Normal data
+				if (Normals.Num() == NumVerts)
+				{
+					ModifyVert.Normal = Normals[VertIdx];
+				}
 
-		if (bPositionsChanging)
-		{
-			MarkRenderStateDirty();
-		}
-		else if (SceneProxy)
-		{
-			// Create data to update section
-			FProcMeshSectionUpdateData* SectionData = new FProcMeshSectionUpdateData;
-			SectionData->TargetSection = SectionIndex;
-			SectionData->NewVertexBuffer = Section.ProcVertexBuffer;
+				// Tangent data
+				if (Tangents.Num() == NumVerts)
+				{
+					ModifyVert.Tangent = Tangents[VertIdx];
+				}
 
-			// Enqueue command to send to render thread
-			FProceduralMeshSceneProxy* ProcMeshSceneProxy = (FProceduralMeshSceneProxy*)SceneProxy;
-			ENQUEUE_RENDER_COMMAND(FProcMeshSectionUpdate)(
-				[ProcMeshSceneProxy, SectionData](FRHICommandListImmediate& RHICmdList)
+				// UV0 data
+				if (UV0.Num() == NumVerts)
 				{
-					ProcMeshSceneProxy->UpdateSection_RenderThread(SectionData);
+					ModifyVert.UV0 = UV0[VertIdx];
+				}
+				// UV1 data
+				if (UV1.Num() == NumVerts)
+				{
+					ModifyVert.UV1 = UV1[VertIdx];
+				}
+				// UV2 data
+				if (UV2.Num() == NumVerts)
+				{
+					ModifyVert.UV2 = UV2[VertIdx];
+				}
+				// UV3 data
+				if (UV3.Num() == NumVerts)
+				{
+					ModifyVert.UV3 = UV3[VertIdx];
 				}
-			);
-		}
 
-		// If we have collision enabled on this section, update that too
-		if(bPositionsChanging && Section.bEnableCollision)
-		{
-			TArray<FVector> CollisionPositions;
+				// Color data
+				if (VertexColors.Num() == NumVerts)
+				{
+					ModifyVert.Color = VertexColors[VertIdx];
+				}
+			}
 
-			// We have one collision mesh for all sections, so need to build array of _all_ positions
-			for (const FProcMeshSection& CollisionSection : ProcMeshSections)
+			// If we have collision enabled on this section, update that too
+			if (Section.bEnableCollision)
 			{
-				// If section has collision, copy it
-				if (CollisionSection.bEnableCollision)
+				TArray<FVector> CollisionPositions;
+
+				// We have one collision mesh for all sections, so need to build array of _all_ positions
+				for (const FProcMeshSection& CollisionSection : ProcMeshSections)
 				{
-					for (int32 VertIdx = 0; VertIdx < CollisionSection.ProcVertexBuffer.Num(); VertIdx++)
+					// If section has collision, copy it
+					if (CollisionSection.bEnableCollision)
 					{
-						CollisionPositions.Add(CollisionSection.ProcVertexBuffer[VertIdx].Position);
+						for (int32 VertIdx = 0; VertIdx < CollisionSection.ProcVertexBuffer.Num(); VertIdx++)
+						{
+							CollisionPositions.Add(CollisionSection.ProcVertexBuffer[VertIdx].Position);
+						}
 					}
 				}
+
+				// Pass new positions to trimesh
+				BodyInstance.UpdateTriMeshVertices(CollisionPositions);
 			}
 
-			// Pass new positions to trimesh
-			BodyInstance.UpdateTriMeshVertices(CollisionPositions);
-		}
+			// If we have a valid proxy and it is not pending recreation
+			if (SceneProxy && !IsRenderStateDirty())
+			{
+				// Create data to update section
+				FProcMeshSectionUpdateData* SectionData = new FProcMeshSectionUpdateData;
+				SectionData->TargetSection = SectionIndex;
+				SectionData->NewVertexBuffer = Section.ProcVertexBuffer;
+
+				// Enqueue command to send to render thread
+				FProceduralMeshSceneProxy* ProcMeshSceneProxy = (FProceduralMeshSceneProxy*)SceneProxy;
+				ENQUEUE_RENDER_COMMAND(FProcMeshSectionUpdate)
+				([ProcMeshSceneProxy, SectionData](FRHICommandListImmediate& RHICmdList) { ProcMeshSceneProxy->UpdateSection_RenderThread(SectionData); });
+			}
 
-		if (bPositionsChanging)
+			UpdateLocalBounds();		 // Update overall bounds
+			MarkRenderTransformDirty();  // Need to send new bounds to render thread
+		}
+		else
 		{
-			UpdateLocalBounds(); // Update overall bounds
-			MarkRenderTransformDirty(); // Need to send new bounds to render thread
+			UE_LOG(LogProceduralComponent, Error, TEXT("Trying to update a procedural mesh component section with a different number of vertices [Previous: %i, New: %i] (clear and recreate mesh section instead)"), PreviousNumVerts, NumVerts);
 		}
 	}
 }
