Cherry-pick 6961133: Added ray tracing support for procedural meshes

diff --git a/Engine/Plugins/Runtime/ProceduralMeshComponent/Source/ProceduralMeshComponent/Private/ProceduralMeshComponent.cpp b/Engine/Plugins/Runtime/ProceduralMeshComponent/Source/ProceduralMeshComponent/Private/ProceduralMeshComponent.cpp
--- a/Engine/Plugins/Runtime/ProceduralMeshComponent/Source/ProceduralMeshComponent/Private/ProceduralMeshComponent.cpp
+++ b/Engine/Plugins/Runtime/ProceduralMeshComponent/Source/ProceduralMeshComponent/Private/ProceduralMeshComponent.cpp
@@ -18,6 +18,8 @@
 #include "DynamicMeshBuilder.h"
 #include "PhysicsEngine/PhysicsSettings.h"
 #include "StaticMeshResources.h"
+#include "RayTracingDefinitions.h"
+#include "RayTracingInstance.h"
 
 DECLARE_CYCLE_STAT(TEXT("Create ProcMesh Proxy"), STAT_ProcMesh_CreateSceneProxy, STATGROUP_ProceduralMesh);
 DECLARE_CYCLE_STAT(TEXT("Create Mesh Section"), STAT_ProcMesh_CreateMeshSection, STATGROUP_ProceduralMesh);
@@ -65,6 +67,10 @@ public:
 	/** Whether this section is currently visible */
 	bool bSectionVisible;
 
+#if RHI_RAYTRACING
+	FRayTracingGeometry RayTracingGeometry;
+#endif
+
 	FProcMeshProxySection(ERHIFeatureLevel::Type InFeatureLevel)
 	: Material(NULL)
 	, VertexFactory(InFeatureLevel, "FProcMeshProxySection")
@@ -162,6 +168,37 @@ public:
 
 				// Save ref to new section
 				Sections[SectionIdx] = NewSection;
+
+#if RHI_RAYTRACING
+				if (IsRayTracingEnabled())
+				{
+					ENQUEUE_RENDER_COMMAND(InitProceduralMeshRayTracingGeometry)(
+						[this, NewSection/*, VertexBufferRHI, IndexBufferRHI, VertexBufferStride, TrianglesCount, RenderSections*/](FRHICommandListImmediate& RHICmdList)
+					{
+						FRayTracingGeometryInitializer Initializer;
+						Initializer.PositionVertexBuffer = nullptr;
+						Initializer.IndexBuffer = nullptr;
+						Initializer.BaseVertexIndex = 0;
+						Initializer.VertexBufferStride = 12;
+						Initializer.VertexBufferByteOffset = 0;
+						Initializer.TotalPrimitiveCount = 0;
+						Initializer.VertexBufferElementType = VET_Float3;
+						Initializer.GeometryType = RTGT_Triangles;
+						Initializer.bFastBuild = true;
+						Initializer.bAllowUpdate = false;
+						NewSection->RayTracingGeometry.SetInitializer(Initializer);
+						NewSection->RayTracingGeometry.InitResource();
+
+						NewSection->RayTracingGeometry.Initializer.PositionVertexBuffer = NewSection->VertexBuffers.PositionVertexBuffer.VertexBufferRHI;
+						NewSection->RayTracingGeometry.Initializer.IndexBuffer = NewSection->IndexBuffer.IndexBufferRHI;
+						NewSection->RayTracingGeometry.Initializer.TotalPrimitiveCount = NewSection->IndexBuffer.Indices.Num() / 3;
+
+						//#dxr_todo: add support for segments?
+						
+						NewSection->RayTracingGeometry.UpdateRHI();
+					});
+				}
+#endif
 			}
 		}
 	}
@@ -177,6 +214,14 @@ public:
 				Section->VertexBuffers.ColorVertexBuffer.ReleaseResource();
 				Section->IndexBuffer.ReleaseResource();
 				Section->VertexFactory.ReleaseResource();
+
+#if RHI_RAYTRACING
+				if (IsRayTracingEnabled())
+				{
+					Section->RayTracingGeometry.ReleaseResource();
+				}
+#endif
+
 				delete Section;
 			}
 		}
@@ -376,6 +421,61 @@ public:
 		return(FPrimitiveSceneProxy::GetAllocatedSize());
 	}
 
+
+#if RHI_RAYTRACING
+	virtual bool IsRayTracingRelevant() const override { return true; }
+
+	virtual void GetDynamicRayTracingInstances(FRayTracingMaterialGatheringContext& Context, TArray<FRayTracingInstance>& OutRayTracingInstances) override final
+	{
+		for (const FProcMeshProxySection* Section : Sections)
+		{
+			if (Section != nullptr && Section->bSectionVisible)
+			{
+				FMaterialRenderProxy* MaterialProxy = Section->Material->GetRenderProxy();
+				
+				if (Section->RayTracingGeometry.RayTracingGeometryRHI.IsValid())
+				{
+					check(Section->RayTracingGeometry.Initializer.PositionVertexBuffer.IsValid());
+					check(Section->RayTracingGeometry.Initializer.IndexBuffer.IsValid());
+
+					FRayTracingInstance RayTracingInstance;
+					RayTracingInstance.Geometry = &Section->RayTracingGeometry;
+					RayTracingInstance.InstanceTransforms.Add(GetLocalToWorld());
+
+					uint32 SectionIdx = 0;
+					FMeshBatch MeshBatch;
+
+					MeshBatch.VertexFactory = &Section->VertexFactory;
+					MeshBatch.SegmentIndex = 0;
+					MeshBatch.MaterialRenderProxy = Section->Material->GetRenderProxy();
+					MeshBatch.ReverseCulling = IsLocalToWorldDeterminantNegative();
+					MeshBatch.Type = PT_TriangleList;
+					MeshBatch.DepthPriorityGroup = SDPG_World;
+					MeshBatch.bCanApplyViewModeOverrides = false;
+
+					FMeshBatchElement& BatchElement = MeshBatch.Elements[0];
+					BatchElement.IndexBuffer = &Section->IndexBuffer;
+					//BatchElement.PrimitiveUniformBuffer = DynamicPrimitiveUniformBuffer.UniformBuffer.GetUniformBufferRHI();					
+					//FDynamicPrimitiveUniformBuffer DynamicPrimitiveUniformBuffer;
+					//DynamicPrimitiveUniformBuffer.Set(GetLocalToWorld(), PreviousLocalToWorld, GetBounds(), GetLocalBounds(), true, bHasPrecomputedVolumetricLightmap, DrawsVelocity(), bOutputVelocity);
+					//BatchElement.PrimitiveUniformBufferResource = &DynamicPrimitiveUniformBuffer.UniformBuffer;
+
+					BatchElement.FirstIndex = 0;
+					BatchElement.NumPrimitives = Section->IndexBuffer.Indices.Num() / 3;
+					BatchElement.MinVertexIndex = 0;
+					BatchElement.MaxVertexIndex = Section->VertexBuffers.PositionVertexBuffer.GetNumVertices() - 1;
+
+					RayTracingInstance.Materials.Add(MeshBatch);
+
+					RayTracingInstance.BuildInstanceMaskAndFlags();
+					OutRayTracingInstances.Add(RayTracingInstance);
+				}
+			}
+		}
+	}
+	
+#endif
+
 private:
 	/** Array of sections */
 	TArray<FProcMeshProxySection*> Sections;
diff --git a/Engine/Plugins/Runtime/ProceduralMeshComponent/Source/ProceduralMeshComponent/ProceduralMeshComponent.Build.cs b/Engine/Plugins/Runtime/ProceduralMeshComponent/Source/ProceduralMeshComponent/ProceduralMeshComponent.Build.cs
--- a/Engine/Plugins/Runtime/ProceduralMeshComponent/Source/ProceduralMeshComponent/ProceduralMeshComponent.Build.cs
+++ b/Engine/Plugins/Runtime/ProceduralMeshComponent/Source/ProceduralMeshComponent/ProceduralMeshComponent.Build.cs
@@ -6,7 +6,9 @@ namespace UnrealBuildTool.Rules
 	{
         public ProceduralMeshComponent(ReadOnlyTargetRules Target) : base(Target)
 		{
-			PublicDependencyModuleNames.AddRange(
+            PrivateIncludePaths.Add("../../../../Shaders/Shared");
+
+            PublicDependencyModuleNames.AddRange(
 				new string[]
 				{
 					"Core",
