Cherry-pick 6974791: Fixed potential binding error of the uniform buffer used by procedural mesh sections

diff --git a/Engine/Plugins/Runtime/ProceduralMeshComponent/Source/ProceduralMeshComponent/Private/ProceduralMeshComponent.cpp b/Engine/Plugins/Runtime/ProceduralMeshComponent/Source/ProceduralMeshComponent/Private/ProceduralMeshComponent.cpp
--- a/Engine/Plugins/Runtime/ProceduralMeshComponent/Source/ProceduralMeshComponent/Private/ProceduralMeshComponent.cpp
+++ b/Engine/Plugins/Runtime/ProceduralMeshComponent/Source/ProceduralMeshComponent/Private/ProceduralMeshComponent.cpp
@@ -455,10 +455,16 @@ public:
 
 					FMeshBatchElement& BatchElement = MeshBatch.Elements[0];
 					BatchElement.IndexBuffer = &Section->IndexBuffer;
-					//BatchElement.PrimitiveUniformBuffer = DynamicPrimitiveUniformBuffer.UniformBuffer.GetUniformBufferRHI();					
-					//FDynamicPrimitiveUniformBuffer DynamicPrimitiveUniformBuffer;
-					//DynamicPrimitiveUniformBuffer.Set(GetLocalToWorld(), PreviousLocalToWorld, GetBounds(), GetLocalBounds(), true, bHasPrecomputedVolumetricLightmap, DrawsVelocity(), bOutputVelocity);
-					//BatchElement.PrimitiveUniformBufferResource = &DynamicPrimitiveUniformBuffer.UniformBuffer;
+
+					bool bHasPrecomputedVolumetricLightmap;
+					FMatrix PreviousLocalToWorld;
+					int32 SingleCaptureIndex;
+					bool bOutputVelocity;
+					GetScene().GetPrimitiveUniformShaderParameters_RenderThread(GetPrimitiveSceneInfo(), bHasPrecomputedVolumetricLightmap, PreviousLocalToWorld, SingleCaptureIndex, bOutputVelocity);
+
+					FDynamicPrimitiveUniformBuffer& DynamicPrimitiveUniformBuffer = Context.RayTracingMeshResourceCollector.AllocateOneFrameResource<FDynamicPrimitiveUniformBuffer>();
+					DynamicPrimitiveUniformBuffer.Set(GetLocalToWorld(), PreviousLocalToWorld, GetBounds(), GetLocalBounds(), true, bHasPrecomputedVolumetricLightmap, DrawsVelocity(), bOutputVelocity);
+					BatchElement.PrimitiveUniformBufferResource = &DynamicPrimitiveUniformBuffer.UniformBuffer;
 
 					BatchElement.FirstIndex = 0;
 					BatchElement.NumPrimitives = Section->IndexBuffer.Indices.Num() / 3;
