Crash when trying to "Update Mesh Section" on any procedural mesh just after "Clear Mesh Section" and "Create Mesh Section" for different vertices and triangles.
#fix Ensure that we recreate the scene proxy whener a mesh section is updated with a larger amount of vertices than original, this will reinit the vertex buffers to the correct sizes

diff --git a/Engine/Plugins/Runtime/ProceduralMeshComponent/Source/ProceduralMeshComponent/Private/ProceduralMeshComponent.cpp b/Engine/Plugins/Runtime/ProceduralMeshComponent/Source/ProceduralMeshComponent/Private/ProceduralMeshComponent.cpp
--- a/Engine/Plugins/Runtime/ProceduralMeshComponent/Source/ProceduralMeshComponent/Private/ProceduralMeshComponent.cpp
+++ b/Engine/Plugins/Runtime/ProceduralMeshComponent/Source/ProceduralMeshComponent/Private/ProceduralMeshComponent.cpp
@@ -548,7 +548,11 @@ void UProceduralMeshComponent::UpdateMeshSection(int32 SectionIndex, const TArra
 			}
 		}
 
-		if (SceneProxy)
+		if (bPositionsChanging)
+		{
+			MarkRenderStateDirty();
+		}
+		else if (SceneProxy)
 		{
 			// Create data to update section
 			FProcMeshSectionUpdateData* SectionData = new FProcMeshSectionUpdateData;
