Deprecate FRawMesh - ProceduralMeshComponent

diff --git a/Engine/Plugins/Runtime/ProceduralMeshComponent/Source/ProceduralMeshComponentEditor/Private/ProceduralMeshComponentDetails.cpp b/Engine/Plugins/Runtime/ProceduralMeshComponent/Source/ProceduralMeshComponentEditor/Private/ProceduralMeshComponentDetails.cpp
--- a/Engine/Plugins/Runtime/ProceduralMeshComponent/Source/ProceduralMeshComponentEditor/Private/ProceduralMeshComponentDetails.cpp
+++ b/Engine/Plugins/Runtime/ProceduralMeshComponent/Source/ProceduralMeshComponentEditor/Private/ProceduralMeshComponentDetails.cpp
@@ -16,7 +16,12 @@
 #include "DetailCategoryBuilder.h"
 #include "IDetailsView.h"
 #include "ProceduralMeshComponent.h"
-#include "RawMesh.h"
+//#include "RawMesh.h"
+
+#include "MeshDescription.h"
+#include "MeshAttributes.h"
+#include "MeshAttributeArray.h"
+#include "MeshDescriptionOperations.h"
 
 #include "Dialogs/DlgPickAssetPath.h"
 #include "AssetRegistryModule.h"
@@ -115,65 +120,141 @@ FReply FProceduralMeshComponentDetails::ClickedOnConvertToStaticMesh()
 				MeshName = *Name;
 			}
 
-			// Raw mesh data we are filling in
-			FRawMesh RawMesh;
-			// Materials to apply to new mesh
-			TArray<UMaterialInterface*> MeshMaterials;
 
+			FMeshDescription MeshDescription;
+			UStaticMesh::RegisterMeshAttributes(MeshDescription);
+			FStaticMeshDescriptionAttributeGetter AttributeGetter(&MeshDescription);
+			TPolygonGroupAttributesRef<FName> PolygonGroupNames = AttributeGetter.GetPolygonGroupImportedMaterialSlotNames();
+			TVertexAttributesRef<FVector> VertexPositions = AttributeGetter.GetPositions();
+			TVertexInstanceAttributesRef<FVector> Tangents = AttributeGetter.GetTangents();
+			TVertexInstanceAttributesRef<float> BinormalSigns = AttributeGetter.GetBinormalSigns();
+			TVertexInstanceAttributesRef<FVector> Normals = AttributeGetter.GetNormals();
+			TVertexInstanceAttributesRef<FVector4> Colors = AttributeGetter.GetColors();
+			TVertexInstanceAttributesRef<FVector2D> UVs = AttributeGetter.GetUVs();
+			TEdgeAttributesRef<bool> EdgeHardnesses = AttributeGetter.GetEdgeHardnesses();
+			TEdgeAttributesRef<float> EdgeCreaseSharpnesses = AttributeGetter.GetEdgeCreaseSharpnesses();
+
+			// Materials to apply to new mesh
 			const int32 NumSections = ProcMeshComp->GetNumSections();
-			int32 VertexBase = 0;
+			int32 VertexCount = 0;
+			int32 VertexInstanceCount = 0;
+			int32 PolygonCount = 0;
+			TMap<UMaterialInterface*, FPolygonGroupID> UniqueMaterials;
+			UniqueMaterials.Reserve(NumSections);
+			TArray<FPolygonGroupID> MaterialRemap;
+			MaterialRemap.Reserve(NumSections);
+			//Get all the info we need to create the MeshDescription
 			for (int32 SectionIdx = 0; SectionIdx < NumSections; SectionIdx++)
 			{
 				FProcMeshSection* ProcSection = ProcMeshComp->GetProcMeshSection(SectionIdx);
-
-				// Copy verts
-				for (FProcMeshVertex& Vert : ProcSection->ProcVertexBuffer)
+				VertexCount += ProcSection->ProcVertexBuffer.Num();
+				VertexInstanceCount += ProcSection->ProcIndexBuffer.Num();
+				PolygonCount += ProcSection->ProcIndexBuffer.Num() / 3;
+				UMaterialInterface*Material = ProcMeshComp->GetMaterial(SectionIdx);
+				if (!UniqueMaterials.Contains(Material))
 				{
-					RawMesh.VertexPositions.Add(Vert.Position);
+					FPolygonGroupID NewPolygonGroup = MeshDescription.CreatePolygonGroup();
+					UniqueMaterials.Add(Material, NewPolygonGroup);
+					PolygonGroupNames[NewPolygonGroup] = Material->GetFName();
 				}
-
-				// Copy 'wedge' info
+				FPolygonGroupID* PolygonGroupID = UniqueMaterials.Find(Material);
+				check(PolygonGroupID != nullptr);
+				MaterialRemap.Add(*PolygonGroupID);
+			}
+			MeshDescription.ReserveNewVertices(VertexCount);
+			MeshDescription.ReserveNewVertexInstances(VertexInstanceCount);
+			MeshDescription.ReserveNewPolygons(PolygonCount);
+			MeshDescription.ReserveNewEdges(PolygonCount * 2);
+			UVs.SetNumIndices(4);
+			//Add Vertex and VertexInstance and polygon for each section
+			for (int32 SectionIdx = 0; SectionIdx < NumSections; SectionIdx++)
+			{
+				FProcMeshSection* ProcSection = ProcMeshComp->GetProcMeshSection(SectionIdx);
+				FPolygonGroupID PolygonGroupID = MaterialRemap[SectionIdx];
+				//Create the vertex
+				int32 NumVertex = ProcSection->ProcVertexBuffer.Num();
+				TMap<int32, FVertexID> VertexIndexToVertexID;
+				VertexIndexToVertexID.Reserve(NumVertex); 
+				for (int32 VertexIndex = 0; VertexIndex < NumVertex; ++VertexIndex)
+				{
+					FProcMeshVertex& Vert = ProcSection->ProcVertexBuffer[VertexIndex];
+					const FVertexID VertexID = MeshDescription.CreateVertex();
+					VertexPositions[VertexID] = Vert.Position;
+					VertexIndexToVertexID.Add(VertexIndex, VertexID);
+				}
+				//Create the VertexInstance
 				int32 NumIndices = ProcSection->ProcIndexBuffer.Num();
-				for (int32 IndexIdx=0; IndexIdx < NumIndices; IndexIdx++)
+				int32 NumTri = NumIndices / 3;
+				TMap<int32, FVertexInstanceID> IndiceIndexToVertexInstanceID;
+				IndiceIndexToVertexInstanceID.Reserve(NumVertex);
+				for (int32 IndiceIndex = 0; IndiceIndex < NumIndices; IndiceIndex++)
 				{
-					int32 Index = ProcSection->ProcIndexBuffer[IndexIdx];
+					const int32 VertexIndex = ProcSection->ProcIndexBuffer[IndiceIndex];
+					const FVertexID VertexID = VertexIndexToVertexID[VertexIndex];
+					const FVertexInstanceID VertexInstanceID = MeshDescription.CreateVertexInstance(VertexID);
+					IndiceIndexToVertexInstanceID.Add(IndiceIndex, VertexInstanceID);
 
-					RawMesh.WedgeIndices.Add(Index + VertexBase);
+					FProcMeshVertex& ProcVertex = ProcSection->ProcVertexBuffer[VertexIndex];
 
-					FProcMeshVertex& ProcVertex = ProcSection->ProcVertexBuffer[Index];
+					Tangents[VertexInstanceID] = ProcVertex.Tangent.TangentX;
+					Normals[VertexInstanceID] = ProcVertex.Normal;
+					BinormalSigns[VertexInstanceID] = ProcVertex.Tangent.bFlipTangentY ? -1.f : 1.f;
 
-					FVector TangentX = ProcVertex.Tangent.TangentX;
-					FVector TangentZ = ProcVertex.Normal;
-					FVector TangentY = (TangentX ^ TangentZ).GetSafeNormal() * (ProcVertex.Tangent.bFlipTangentY ? -1.f : 1.f);
+					Colors[VertexInstanceID] = FLinearColor(ProcVertex.Color);
 
-					RawMesh.WedgeTangentX.Add(TangentX);
-					RawMesh.WedgeTangentY.Add(TangentY);
-					RawMesh.WedgeTangentZ.Add(TangentZ);
-
-					RawMesh.WedgeTexCoords[0].Add(ProcVertex.UV0);
-					RawMesh.WedgeTexCoords[1].Add(ProcVertex.UV1);
-					RawMesh.WedgeTexCoords[2].Add(ProcVertex.UV2);
-					RawMesh.WedgeTexCoords[3].Add(ProcVertex.UV3);
-					RawMesh.WedgeColors.Add(ProcVertex.Color);
+					UVs.Set(VertexInstanceID, 0, ProcVertex.UV0);
+					UVs.Set(VertexInstanceID, 1, ProcVertex.UV1);
+					UVs.Set(VertexInstanceID, 2, ProcVertex.UV2);
+					UVs.Set(VertexInstanceID, 3, ProcVertex.UV3);
 				}
-
-				// copy face info
-				int32 NumTris = NumIndices / 3;
-				for (int32 TriIdx=0; TriIdx < NumTris; TriIdx++)
+				
+				//Create the polygons for this section
+				for (int32 TriIdx = 0; TriIdx < NumTri; TriIdx++)
 				{
-					RawMesh.FaceMaterialIndices.Add(SectionIdx);
-					RawMesh.FaceSmoothingMasks.Add(0); // Assume this is ignored as bRecomputeNormals is false
+					FVertexID VertexIndexes[3];
+					FVertexInstanceID VertexInstanceIDs[3];
+					for (int32 CornerIndex = 0; CornerIndex < 3; ++CornerIndex)
+					{
+						const int32 IndiceIndex = (TriIdx * 3) + CornerIndex;
+						const int32 VertexIndex = ProcSection->ProcIndexBuffer[IndiceIndex];
+						VertexIndexes[CornerIndex] = VertexIndexToVertexID[VertexIndex];
+						VertexInstanceIDs[CornerIndex] = IndiceIndexToVertexInstanceID[IndiceIndex];
+					}
+
+					TArray<FMeshDescription::FContourPoint> Contours;
+					for (int32 CornerIndex = 0; CornerIndex < 3; ++CornerIndex)
+					{
+						int32 ContourPointIndex = Contours.AddDefaulted();
+						FMeshDescription::FContourPoint& ContourPoint = Contours[ContourPointIndex];
+						//Find the matching edge ID
+						uint32 CornerIndices[2];
+						CornerIndices[0] = (CornerIndex + 0) % 3;
+						CornerIndices[1] = (CornerIndex + 1) % 3;
+
+						FVertexID EdgeVertexIDs[2];
+						EdgeVertexIDs[0] = VertexIndexes[CornerIndices[0]];
+						EdgeVertexIDs[1] = VertexIndexes[CornerIndices[1]];
+
+						FEdgeID MatchEdgeId = MeshDescription.GetVertexPairEdge(EdgeVertexIDs[0], EdgeVertexIDs[1]);
+						if (MatchEdgeId == FEdgeID::Invalid)
+						{
+							MatchEdgeId = MeshDescription.CreateEdge(EdgeVertexIDs[0], EdgeVertexIDs[1]);
+							EdgeHardnesses[MatchEdgeId] = false;
+							EdgeCreaseSharpnesses[MatchEdgeId] = 0.0f;
+						}
+						ContourPoint.EdgeID = MatchEdgeId;
+						ContourPoint.VertexInstanceID = VertexInstanceIDs[CornerIndices[0]];
+					}
+					// Insert a polygon into the mesh
+					const FPolygonID NewPolygonID = MeshDescription.CreatePolygon(PolygonGroupID, Contours);
+					//Triangulate the polygon
+					FMeshPolygon& Polygon = MeshDescription.GetPolygon(NewPolygonID);
+					MeshDescription.ComputePolygonTriangulation(NewPolygonID, Polygon.Triangles);
 				}
-
-				// Remember material
-				MeshMaterials.Add(ProcMeshComp->GetMaterial(SectionIdx));
-
-				// Update offset for creating one big index/vertex buffer
-				VertexBase += ProcSection->ProcVertexBuffer.Num();
 			}
 
 			// If we got some valid data.
-			if (RawMesh.VertexPositions.Num() >= 3 && RawMesh.WedgeIndices.Num() >= 3)
+			if (MeshDescription.Polygons().Num() > 0)
 			{
 				// Then find/create it.
 				UPackage* Package = CreatePackage(NULL, *UserPackageName);
@@ -195,12 +276,19 @@ FReply FProceduralMeshComponentDetails::ClickedOnConvertToStaticMesh()
 				SrcModel.BuildSettings.bGenerateLightmapUVs = true;
 				SrcModel.BuildSettings.SrcLightmapIndex = 0;
 				SrcModel.BuildSettings.DstLightmapIndex = 1;
-				SrcModel.SaveRawMesh(RawMesh);
+				FMeshDescription* OriginalMeshDescription = StaticMesh->GetOriginalMeshDescription(0);
+				if (OriginalMeshDescription == nullptr)
+				{
+					OriginalMeshDescription = StaticMesh->CreateOriginalMeshDescription(0);
+				}
+				*OriginalMeshDescription = MeshDescription;
+				StaticMesh->CommitOriginalMeshDescription(0);
 
 				// Copy materials to new mesh
-				for (UMaterialInterface* Material : MeshMaterials)
+				for (auto Kvp : UniqueMaterials)
 				{
-					StaticMesh->StaticMaterials.Add(FStaticMaterial(Material));
+					UMaterialInterface* Material = Kvp.Key;
+					StaticMesh->StaticMaterials.Add(FStaticMaterial(Material, Material->GetFName(), Material->GetFName()));
 				}
 
 				//Set the Imported version before calling the build
diff --git a/Engine/Plugins/Runtime/ProceduralMeshComponent/Source/ProceduralMeshComponentEditor/ProceduralMeshComponentEditor.Build.cs b/Engine/Plugins/Runtime/ProceduralMeshComponent/Source/ProceduralMeshComponentEditor/ProceduralMeshComponentEditor.Build.cs
--- a/Engine/Plugins/Runtime/ProceduralMeshComponent/Source/ProceduralMeshComponentEditor/ProceduralMeshComponentEditor.Build.cs
+++ b/Engine/Plugins/Runtime/ProceduralMeshComponent/Source/ProceduralMeshComponentEditor/ProceduralMeshComponentEditor.Build.cs
@@ -20,7 +20,8 @@ namespace UnrealBuildTool.Rules
                     "ShaderCore",
                     "RHI",
                     "ProceduralMeshComponent",
-                    "RawMesh",
+                    "MeshDescription",
+                    "MeshDescriptionOperations",
                     "AssetTools",
                     "AssetRegistry"
                 }
diff --git a/Engine/Source/Runtime/Engine/Classes/Engine/StaticMesh.h b/Engine/Source/Runtime/Engine/Classes/Engine/StaticMesh.h
--- a/Engine/Source/Runtime/Engine/Classes/Engine/StaticMesh.h
+++ b/Engine/Source/Runtime/Engine/Classes/Engine/StaticMesh.h
@@ -19,15 +19,20 @@
 #include "Templates/UniquePtr.h"
 #include "StaticMeshResources.h"
 #include "PerPlatformProperties.h"
+
+#include "MeshDescription.h"
+#include "MeshAttributes.h"
+#include "MeshAttributeArray.h"
+
 #include "StaticMesh.generated.h"
 
 class FSpeedTreeWind;
 class UAssetUserData;
 class UMaterialInterface;
 class UNavCollisionBase;
-struct FMeshDescription;
 struct FStaticMeshLODResources;
 
+
 /*-----------------------------------------------------------------------------
 	Legacy mesh optimization settings.
 -----------------------------------------------------------------------------*/
@@ -444,6 +449,42 @@ struct FMaterialRemapIndex
 	TArray<int32> MaterialRemap;
 };
 
+struct FStaticMeshDescriptionAttributeGetter
+{
+	ENGINE_API FStaticMeshDescriptionAttributeGetter(FMeshDescription* InMeshDescription)
+		: MeshDescription(InMeshDescription)
+	{}
+
+	FMeshDescription* MeshDescription;
+
+	TVertexAttributesRef<FVector> GetPositions() const { return MeshDescription->VertexAttributes().GetAttributesRef<FVector>(MeshAttribute::Vertex::Position); }
+	TVertexAttributesConstRef<FVector> GetPositionsConst() const { return MeshDescription->VertexAttributes().GetAttributesRef<FVector>(MeshAttribute::Vertex::Position); }
+
+	TVertexInstanceAttributesRef<FVector> GetNormals() const { return MeshDescription->VertexInstanceAttributes().GetAttributesRef<FVector>(MeshAttribute::VertexInstance::Normal); }
+	TVertexInstanceAttributesConstRef<FVector> GetNormalsConst() const { return MeshDescription->VertexInstanceAttributes().GetAttributesRef<FVector>(MeshAttribute::VertexInstance::Normal); }
+
+	TVertexInstanceAttributesRef<FVector> GetTangents() const { return MeshDescription->VertexInstanceAttributes().GetAttributesRef<FVector>(MeshAttribute::VertexInstance::Tangent); }
+	TVertexInstanceAttributesConstRef<FVector> GetTangentsConst() const { return MeshDescription->VertexInstanceAttributes().GetAttributesRef<FVector>(MeshAttribute::VertexInstance::Tangent); }
+
+	TVertexInstanceAttributesRef<float> GetBinormalSigns() const { return MeshDescription->VertexInstanceAttributes().GetAttributesRef<float>(MeshAttribute::VertexInstance::BinormalSign); }
+	TVertexInstanceAttributesConstRef<float> GetBinormalSignsConst() const { return MeshDescription->VertexInstanceAttributes().GetAttributesRef<float>(MeshAttribute::VertexInstance::BinormalSign); }
+
+	TVertexInstanceAttributesRef<FVector4> GetColors() const { return MeshDescription->VertexInstanceAttributes().GetAttributesRef<FVector4>(MeshAttribute::VertexInstance::Color); }
+	TVertexInstanceAttributesConstRef<FVector4> GetColorsConst() const { return MeshDescription->VertexInstanceAttributes().GetAttributesRef<FVector4>(MeshAttribute::VertexInstance::Color); }
+
+	TVertexInstanceAttributesRef<FVector2D> GetUVs() const { return MeshDescription->VertexInstanceAttributes().GetAttributesRef<FVector2D>(MeshAttribute::VertexInstance::TextureCoordinate); }
+	TVertexInstanceAttributesConstRef<FVector2D> GetUVsConst() const { return MeshDescription->VertexInstanceAttributes().GetAttributesRef<FVector2D>(MeshAttribute::VertexInstance::TextureCoordinate); }
+
+	TEdgeAttributesRef<bool> GetEdgeHardnesses() const { return MeshDescription->EdgeAttributes().GetAttributesRef<bool>(MeshAttribute::Edge::IsHard); }
+	TEdgeAttributesConstRef<bool> GetEdgeHardnessesConst() const { return MeshDescription->EdgeAttributes().GetAttributesRef<bool>(MeshAttribute::Edge::IsHard); }
+	
+	TEdgeAttributesRef<float> GetEdgeCreaseSharpnesses() const { return MeshDescription->EdgeAttributes().GetAttributesRef<float>(MeshAttribute::Edge::CreaseSharpness); }
+	TEdgeAttributesConstRef<float> GetEdgeCreaseSharpnessesConst() const { return MeshDescription->EdgeAttributes().GetAttributesRef<float>(MeshAttribute::Edge::CreaseSharpness); }
+
+	TPolygonGroupAttributesRef<FName> GetPolygonGroupImportedMaterialSlotNames() { return MeshDescription->PolygonGroupAttributes().GetAttributesRef<FName>(MeshAttribute::PolygonGroup::ImportedMaterialSlotName); }
+	TPolygonGroupAttributesConstRef<FName> GetPolygonGroupImportedMaterialSlotNamesConst() { return MeshDescription->PolygonGroupAttributes().GetAttributesRef<FName>(MeshAttribute::PolygonGroup::ImportedMaterialSlotName); }
+};
+
 /**
  * Container UObject which holds a MeshDescription for every SourceModel LOD.
  * Eventually this will be outered to the static mesh's package, so the data is saved along with the static mesh, but will not be referenced by it.
