Remove unwanted comment

diff --git a/Engine/Plugins/Runtime/ProceduralMeshComponent/Source/ProceduralMeshComponentEditor/Private/ProceduralMeshComponentDetails.cpp b/Engine/Plugins/Runtime/ProceduralMeshComponent/Source/ProceduralMeshComponentEditor/Private/ProceduralMeshComponentDetails.cpp
--- a/Engine/Plugins/Runtime/ProceduralMeshComponent/Source/ProceduralMeshComponentEditor/Private/ProceduralMeshComponentDetails.cpp
+++ b/Engine/Plugins/Runtime/ProceduralMeshComponent/Source/ProceduralMeshComponentEditor/Private/ProceduralMeshComponentDetails.cpp
@@ -16,7 +16,6 @@
 #include "DetailCategoryBuilder.h"
 #include "IDetailsView.h"
 #include "ProceduralMeshComponent.h"
-//#include "RawMesh.h"
 
 #include "MeshDescription.h"
 #include "MeshAttributes.h"
