Fixed cast shadows flag was ignored for procedural meshes with ray tracing shadows enabled.

diff --git a/Engine/Plugins/Runtime/ProceduralMeshComponent/Source/ProceduralMeshComponent/Private/ProceduralMeshComponent.cpp b/Engine/Plugins/Runtime/ProceduralMeshComponent/Source/ProceduralMeshComponent/Private/ProceduralMeshComponent.cpp
--- a/Engine/Plugins/Runtime/ProceduralMeshComponent/Source/ProceduralMeshComponent/Private/ProceduralMeshComponent.cpp
+++ b/Engine/Plugins/Runtime/ProceduralMeshComponent/Source/ProceduralMeshComponent/Private/ProceduralMeshComponent.cpp
@@ -451,6 +451,7 @@ public:
 					MeshBatch.Type = PT_TriangleList;
 					MeshBatch.DepthPriorityGroup = SDPG_World;
 					MeshBatch.bCanApplyViewModeOverrides = false;
+					MeshBatch.CastRayTracedShadow = IsShadowCast(Context.ReferenceView);
 
 					FMeshBatchElement& BatchElement = MeshBatch.Elements[0];
 					BatchElement.IndexBuffer = &Section->IndexBuffer;
