Change GENERATED_U{CLASS/STRUCT}_BODY to GENERATE_BODY in Procedural Mesh Component

#rnx

diff --git a/Engine/Plugins/Runtime/ProceduralMeshComponent/Source/ProceduralMeshComponent/Private/KismetProceduralMeshLibrary.cpp b/Engine/Plugins/Runtime/ProceduralMeshComponent/Source/ProceduralMeshComponent/Private/KismetProceduralMeshLibrary.cpp
--- a/Engine/Plugins/Runtime/ProceduralMeshComponent/Source/ProceduralMeshComponent/Private/KismetProceduralMeshLibrary.cpp
+++ b/Engine/Plugins/Runtime/ProceduralMeshComponent/Source/ProceduralMeshComponent/Private/KismetProceduralMeshLibrary.cpp
@@ -17,12 +17,6 @@ DECLARE_CYCLE_STAT(TEXT("Update Collision"), STAT_ProcMesh_CalcTangents, STATGRO
 #define LOCTEXT_NAMESPACE "KismetProcMeshLibrary"
 
 
-UKismetProceduralMeshLibrary::UKismetProceduralMeshLibrary(const FObjectInitializer& ObjectInitializer)
-: Super(ObjectInitializer)
-{
-
-}
-
 void UKismetProceduralMeshLibrary::ConvertQuadToTriangles(TArray<int32>& Triangles, int32 Vert0, int32 Vert1, int32 Vert2, int32 Vert3)
 {
 	Triangles.Add(Vert0);
diff --git a/Engine/Plugins/Runtime/ProceduralMeshComponent/Source/ProceduralMeshComponent/Public/KismetProceduralMeshLibrary.h b/Engine/Plugins/Runtime/ProceduralMeshComponent/Source/ProceduralMeshComponent/Public/KismetProceduralMeshLibrary.h
--- a/Engine/Plugins/Runtime/ProceduralMeshComponent/Source/ProceduralMeshComponent/Public/KismetProceduralMeshLibrary.h
+++ b/Engine/Plugins/Runtime/ProceduralMeshComponent/Source/ProceduralMeshComponent/Public/KismetProceduralMeshLibrary.h
@@ -26,7 +26,8 @@ enum class EProcMeshSliceCapOption : uint8
 UCLASS(meta=(ScriptName="ProceduralMeshLibrary"))
 class PROCEDURALMESHCOMPONENT_API UKismetProceduralMeshLibrary : public UBlueprintFunctionLibrary
 {
-	GENERATED_UCLASS_BODY()
+	GENERATED_BODY()
+public:
 
 	/** Generate vertex and index buffer for a simple box, given the supplied dimensions. Normals, UVs and tangents are also generated for each vertex. */
 	UFUNCTION(BlueprintCallable, Category = "Components|ProceduralMesh")
diff --git a/Engine/Plugins/Runtime/ProceduralMeshComponent/Source/ProceduralMeshComponent/Public/ProceduralMeshComponent.h b/Engine/Plugins/Runtime/ProceduralMeshComponent/Source/ProceduralMeshComponent/Public/ProceduralMeshComponent.h
--- a/Engine/Plugins/Runtime/ProceduralMeshComponent/Source/ProceduralMeshComponent/Public/ProceduralMeshComponent.h
+++ b/Engine/Plugins/Runtime/ProceduralMeshComponent/Source/ProceduralMeshComponent/Public/ProceduralMeshComponent.h
@@ -18,7 +18,8 @@ class FPrimitiveSceneProxy;
 USTRUCT(BlueprintType)
 struct FProcMeshTangent
 {
-	GENERATED_USTRUCT_BODY()
+	GENERATED_BODY()
+public:
 
 	/** Direction of X tangent for this vertex */
 	UPROPERTY(EditAnywhere, BlueprintReadWrite, Category = Tangent)
@@ -48,7 +49,8 @@ struct FProcMeshTangent
 USTRUCT(BlueprintType)
 struct FProcMeshVertex
 {
-	GENERATED_USTRUCT_BODY()
+	GENERATED_BODY()
+public:
 
 	/** Vertex position */
 	UPROPERTY(EditAnywhere, BlueprintReadWrite, Category = Vertex)
@@ -99,7 +101,8 @@ struct FProcMeshVertex
 USTRUCT()
 struct FProcMeshSection
 {
-	GENERATED_USTRUCT_BODY()
+	GENERATED_BODY()
+public:
 
 	/** Vertex buffer for this section */
 	UPROPERTY()
@@ -144,7 +147,10 @@ struct FProcMeshSection
 UCLASS(hidecategories = (Object, LOD), meta = (BlueprintSpawnableComponent), ClassGroup = Rendering)
 class PROCEDURALMESHCOMPONENT_API UProceduralMeshComponent : public UMeshComponent, public IInterface_CollisionDataProvider
 {
-	GENERATED_UCLASS_BODY()
+	GENERATED_BODY()
+public:
+
+	UProceduralMeshComponent(const FObjectInitializer& ObjectInitializer);
 
 	/**
 	 *	Create/replace a section for this procedural mesh component.
