Allow MeshModeling tools to operate on ProceduralMeshComponent

#rnx none

diff --git a/Engine/Plugins/Runtime/ProceduralMeshComponent/Source/ProceduralMeshComponent/Private/ProceduralMeshBridge.cpp b/Engine/Plugins/Runtime/ProceduralMeshComponent/Source/ProceduralMeshComponent/Private/ProceduralMeshBridge.cpp
new file mode 100644
--- /dev/null
+++ b/Engine/Plugins/Runtime/ProceduralMeshComponent/Source/ProceduralMeshComponent/Private/ProceduralMeshBridge.cpp
@@ -0,0 +1,45 @@
+// Copyright 1998-2019 Epic Games, Inc. All Rights Reserved.
+
+#include "ProceduralMeshBridge.h"
+
+#include "MeshDescription.h"
+#include "ProceduralMeshComponent.h"
+#include "ProceduralMeshConversion.h"
+#include "Components/PrimitiveComponent.h"
+
+
+bool
+FProceduralMeshComponentTargetFactory::CanBuild(UActorComponent* Component)
+{
+	return !!Cast<UProceduralMeshComponent>(Component);
+}
+
+TUniquePtr<FPrimitiveComponentTarget>
+FProceduralMeshComponentTargetFactory::Build(UPrimitiveComponent* Component)
+{
+	UProceduralMeshComponent* ProceduralMeshComponent = Cast<UProceduralMeshComponent>(Component);
+	if (ProceduralMeshComponent != nullptr)
+	{
+		return TUniquePtr< FPrimitiveComponentTarget > { new FProceduralMeshComponentTarget{ Component } };
+	}
+	return {};
+}
+
+FProceduralMeshComponentTarget::FProceduralMeshComponentTarget( UPrimitiveComponent* Component )
+	: FPrimitiveComponentTarget{ Cast<UProceduralMeshComponent>(Component) }, MeshDescription{new FMeshDescription}
+{
+}
+
+FMeshDescription*
+FProceduralMeshComponentTarget::GetMesh()
+{
+	*MeshDescription = BuildMeshDescription( Cast<UProceduralMeshComponent>(Component) );
+	return MeshDescription.Get();
+}
+
+void
+FProceduralMeshComponentTarget::CommitMesh( const FCommitter& ModifyFunc )
+{
+	ModifyFunc(MeshDescription.Get());
+	MeshDescriptionToProcMesh( *MeshDescription, Cast<UProceduralMeshComponent>(Component) );
+}
diff --git a/Engine/Plugins/Runtime/ProceduralMeshComponent/Source/ProceduralMeshComponent/Private/ProceduralMeshComponentPlugin.cpp b/Engine/Plugins/Runtime/ProceduralMeshComponent/Source/ProceduralMeshComponent/Private/ProceduralMeshComponentPlugin.cpp
--- a/Engine/Plugins/Runtime/ProceduralMeshComponent/Source/ProceduralMeshComponent/Private/ProceduralMeshComponentPlugin.cpp
+++ b/Engine/Plugins/Runtime/ProceduralMeshComponent/Source/ProceduralMeshComponent/Private/ProceduralMeshComponentPlugin.cpp
@@ -1,12 +1,11 @@
 // Copyright 1998-2019 Epic Games, Inc. All Rights Reserved.
 
 #include "CoreMinimal.h"
+#include "ComponentSourceInterfaces.h"
 #include "Modules/ModuleManager.h"
+#include "ProceduralMeshBridge.h"
 #include "IProceduralMeshComponentPlugin.h"
 
-
-
-
 class FProceduralMeshComponentPlugin : public IProceduralMeshComponentPlugin
 {
 	/** IModuleInterface implementation */
@@ -20,7 +19,8 @@ IMPLEMENT_MODULE( FProceduralMeshComponentPlugin, ProceduralMeshComponent )
 
 void FProceduralMeshComponentPlugin::StartupModule()
 {
-	
+	// This code will execute after your module is loaded into memory; the exact timing is specified in the .uplugin file per-module
+	AddComponentTargetFactory( TUniquePtr<FComponentTargetFactory>{new FProceduralMeshComponentTargetFactory{} } );
 }
 
 
diff --git a/Engine/Plugins/Runtime/ProceduralMeshComponent/Source/ProceduralMeshComponent/Private/ProceduralMeshConversion.cpp b/Engine/Plugins/Runtime/ProceduralMeshComponent/Source/ProceduralMeshComponent/Private/ProceduralMeshConversion.cpp
new file mode 100644
--- /dev/null
+++ b/Engine/Plugins/Runtime/ProceduralMeshComponent/Source/ProceduralMeshComponent/Private/ProceduralMeshConversion.cpp
@@ -0,0 +1,305 @@
+// Copyright 1998-2019 Epic Games, Inc. All Rights Reserved.
+
+#include "ProceduralMeshConversion.h"
+#include "Engine/StaticMesh.h"
+#include "ProceduralMeshComponent.h"
+#include "MeshAttributes.h"
+#include "MeshDescription.h"
+
+struct FProceduralMeshDescriptionAttributeGetter
+{
+	FProceduralMeshDescriptionAttributeGetter(FMeshDescription* InMeshDescription)
+		: Mesh(InMeshDescription)
+	{}
+
+	FMeshDescription* Mesh;
+
+	TVertexAttributesRef<FVector>
+	GetPositions() const
+	{
+		return Mesh->VertexAttributes().GetAttributesRef<FVector>(MeshAttribute::Vertex::Position);
+	}
+	TVertexInstanceAttributesRef<FVector>
+	GetNormals() const
+	{
+		return Mesh->VertexInstanceAttributes().GetAttributesRef<FVector>(MeshAttribute::VertexInstance::Normal);
+	}
+	TVertexInstanceAttributesRef<FVector>
+	GetTangents() const
+	{
+		return Mesh->VertexInstanceAttributes().GetAttributesRef<FVector>(MeshAttribute::VertexInstance::Tangent);
+	}
+	TVertexInstanceAttributesRef<float>
+	GetBinormalSigns() const {
+		return Mesh->VertexInstanceAttributes().GetAttributesRef<float>(MeshAttribute::VertexInstance::BinormalSign);
+	}
+	TVertexInstanceAttributesRef<FVector4>
+	GetColors() const
+	{
+		return Mesh->VertexInstanceAttributes().GetAttributesRef<FVector4>(MeshAttribute::VertexInstance::Color);
+	}
+	TVertexInstanceAttributesRef<FVector2D>
+	GetUVs() const
+	{
+		return Mesh->VertexInstanceAttributes().GetAttributesRef<FVector2D>(MeshAttribute::VertexInstance::TextureCoordinate);
+	}
+	TPolygonGroupAttributesRef<FName>
+	GetPolygonGroupImportedMaterialSlotNames()
+	{
+		return Mesh->PolygonGroupAttributes().GetAttributesRef<FName>(MeshAttribute::PolygonGroup::ImportedMaterialSlotName);
+	}
+};
+
+struct FProceduralMeshDescriptionConstAttributeGetter
+{
+	FProceduralMeshDescriptionConstAttributeGetter(const FMeshDescription* InMeshDescription)
+		: Mesh(InMeshDescription)
+	{}
+
+	const FMeshDescription* Mesh;
+
+	TVertexAttributesConstRef<FVector>
+	GetPositions() const
+	{
+		return Mesh->VertexAttributes().GetAttributesRef<FVector>(MeshAttribute::Vertex::Position);
+	}
+	TVertexInstanceAttributesConstRef<FVector>
+	GetNormals() const
+	{
+		return Mesh->VertexInstanceAttributes().GetAttributesRef<FVector>(MeshAttribute::VertexInstance::Normal);
+	}
+	TVertexInstanceAttributesConstRef<FVector>
+	GetTangents() const
+	{
+		return Mesh->VertexInstanceAttributes().GetAttributesRef<FVector>(MeshAttribute::VertexInstance::Tangent);
+	}
+	TVertexInstanceAttributesConstRef<float>
+	GetBinormalSigns() const {
+		return Mesh->VertexInstanceAttributes().GetAttributesRef<float>(MeshAttribute::VertexInstance::BinormalSign);
+	}
+	TVertexInstanceAttributesConstRef<FVector4>
+	GetColors() const
+	{
+		return Mesh->VertexInstanceAttributes().GetAttributesRef<FVector4>(MeshAttribute::VertexInstance::Color);
+	}
+	TVertexInstanceAttributesConstRef<FVector2D>
+	GetUVs() const
+	{
+		return Mesh->VertexInstanceAttributes().GetAttributesRef<FVector2D>(MeshAttribute::VertexInstance::TextureCoordinate);
+	}
+	TPolygonGroupAttributesConstRef<FName>
+	GetPolygonGroupImportedMaterialSlotNames()
+	{
+		return Mesh->PolygonGroupAttributes().GetAttributesRef<FName>(MeshAttribute::PolygonGroup::ImportedMaterialSlotName);
+	}
+};
+
+TMap<UMaterialInterface*, FPolygonGroupID> BuildMaterialMap(UProceduralMeshComponent* ProcMeshComp, FMeshDescription& MeshDescription)
+{
+	TMap<UMaterialInterface*, FPolygonGroupID> UniqueMaterials;
+	const int32 NumSections = ProcMeshComp->GetNumSections();
+	UniqueMaterials.Reserve(NumSections);
+
+	FProceduralMeshDescriptionAttributeGetter AttributeGetter(&MeshDescription);
+	TPolygonGroupAttributesRef<FName> PolygonGroupNames = AttributeGetter.GetPolygonGroupImportedMaterialSlotNames();
+	for (int32 SectionIdx = 0; SectionIdx < NumSections; SectionIdx++)
+	{
+		FProcMeshSection *ProcSection =
+			ProcMeshComp->GetProcMeshSection(SectionIdx);
+		UMaterialInterface *Material = ProcMeshComp->GetMaterial(SectionIdx);
+		if (!UniqueMaterials.Contains(Material))
+		{
+			FPolygonGroupID NewPolygonGroup = MeshDescription.CreatePolygonGroup();
+			UniqueMaterials.Add(Material, NewPolygonGroup);
+			PolygonGroupNames[NewPolygonGroup] = Material->GetFName();
+		}
+	}
+	return UniqueMaterials;
+}
+
+namespace
+{
+	struct FMeshSectionAttributeData
+	{
+		TArray<FVector> Positions;
+		TArray<int32>   Triangles;
+		TArray<FVector> Normals;
+		TArray<FVector2D> UV0;
+		TArray<FColor> VertexColors;
+		TArray<FProcMeshTangent> Tangents;
+		TMap<FVertexInstanceID, int32> IndexOfInstances;
+	};
+}
+
+void MeshDescriptionToProcMesh( const FMeshDescription& MeshDescription, UProceduralMeshComponent* ProcMeshComp )
+{
+	ProcMeshComp->ClearAllMeshSections();
+
+	FProceduralMeshDescriptionConstAttributeGetter AttributeGetter(&MeshDescription);
+	TPolygonGroupAttributesConstRef<FName> PolygonGroupNames = AttributeGetter.GetPolygonGroupImportedMaterialSlotNames();
+	TVertexAttributesConstRef<FVector> VertexPositions = AttributeGetter.GetPositions();
+	TVertexInstanceAttributesConstRef<FVector> Tangents = AttributeGetter.GetTangents();
+	TVertexInstanceAttributesConstRef<float> BinormalSigns = AttributeGetter.GetBinormalSigns();
+	TVertexInstanceAttributesConstRef<FVector> Normals = AttributeGetter.GetNormals();
+	TVertexInstanceAttributesConstRef<FVector4> Colors = AttributeGetter.GetColors();
+	TVertexInstanceAttributesConstRef<FVector2D> UVs = AttributeGetter.GetUVs();
+
+	const int32 NumSections = ProcMeshComp->GetNumSections();
+
+	TMap<FPolygonGroupID, FMeshSectionAttributeData> GroupedData;
+	for ( const FTriangleID TriangleID : MeshDescription.Triangles().GetElementIDs() )
+	{
+		FPolygonGroupID PolygonGroupID = MeshDescription.GetTrianglePolygonGroup( TriangleID );
+		FMeshSectionAttributeData& SectionData = GroupedData.FindOrAdd(PolygonGroupID);
+		for ( const FVertexInstanceID VertexInstanceID : MeshDescription.GetTriangleVertexInstances( TriangleID ) )
+		{
+			auto* FoundVertexIndex = SectionData.IndexOfInstances.Find( VertexInstanceID );
+			if ( FoundVertexIndex ){
+				SectionData.Triangles.Add( *FoundVertexIndex );
+			}
+			else
+			{
+				FVertexID VertexID = MeshDescription.GetVertexInstanceVertex( VertexInstanceID );
+				int32 NewVertexIndex = SectionData.IndexOfInstances.Num();
+				SectionData.IndexOfInstances.Add( VertexInstanceID,  NewVertexIndex );
+				SectionData.Positions.Add( VertexPositions[VertexID] );
+				SectionData.Triangles.Add( NewVertexIndex );
+				SectionData.Normals.Add( Normals[VertexInstanceID] );
+				SectionData.UV0.Add( UVs[VertexInstanceID] );
+				SectionData.Tangents.Add( FProcMeshTangent( Tangents[VertexInstanceID], BinormalSigns[VertexInstanceID] < 0.f ) );
+			}
+		}
+	}
+	int32 SectionIndex{0};
+	for ( auto& MapEntry : GroupedData )
+	{
+		auto& SectionData = MapEntry.Value;
+		ProcMeshComp->CreateMeshSection(SectionIndex++,
+										SectionData.Positions,
+										SectionData.Triangles,
+										SectionData.Normals,
+										SectionData.UV0,
+										SectionData.VertexColors,
+										SectionData.Tangents,
+										true);
+	}
+}
+
+FMeshDescription BuildMeshDescription( UProceduralMeshComponent* ProcMeshComp )
+{
+	FMeshDescription MeshDescription;
+	UStaticMesh::RegisterMeshAttributes(MeshDescription);
+
+	FProceduralMeshDescriptionAttributeGetter AttributeGetter(&MeshDescription);
+	TPolygonGroupAttributesRef<FName> PolygonGroupNames = AttributeGetter.GetPolygonGroupImportedMaterialSlotNames();
+	TVertexAttributesRef<FVector> VertexPositions = AttributeGetter.GetPositions();
+	TVertexInstanceAttributesRef<FVector> Tangents = AttributeGetter.GetTangents();
+	TVertexInstanceAttributesRef<float> BinormalSigns = AttributeGetter.GetBinormalSigns();
+	TVertexInstanceAttributesRef<FVector> Normals = AttributeGetter.GetNormals();
+	TVertexInstanceAttributesRef<FVector4> Colors = AttributeGetter.GetColors();
+	TVertexInstanceAttributesRef<FVector2D> UVs = AttributeGetter.GetUVs();
+
+	// Materials to apply to new mesh
+	const int32 NumSections = ProcMeshComp->GetNumSections();
+	int32 VertexCount = 0;
+	int32 VertexInstanceCount = 0;
+	int32 PolygonCount = 0;
+
+	TMap<UMaterialInterface*, FPolygonGroupID> UniqueMaterials = BuildMaterialMap(ProcMeshComp, MeshDescription);
+	TArray<FPolygonGroupID> PolygonGroupForSection;
+	PolygonGroupForSection.Reserve(NumSections);
+
+	// Calculate the totals for each ProcMesh element type
+	for (int32 SectionIdx = 0; SectionIdx < NumSections; SectionIdx++)
+	{
+		FProcMeshSection *ProcSection =
+			ProcMeshComp->GetProcMeshSection(SectionIdx);
+		VertexCount += ProcSection->ProcVertexBuffer.Num();
+		VertexInstanceCount += ProcSection->ProcIndexBuffer.Num();
+		PolygonCount += ProcSection->ProcIndexBuffer.Num() / 3;
+	}
+	MeshDescription.ReserveNewVertices(VertexCount);
+	MeshDescription.ReserveNewVertexInstances(VertexInstanceCount);
+	MeshDescription.ReserveNewPolygons(PolygonCount);
+	MeshDescription.ReserveNewEdges(PolygonCount * 2);
+	UVs.SetNumIndices(4);
+
+	// Create the Polygon Groups
+	for (int32 SectionIdx = 0; SectionIdx < NumSections; SectionIdx++)
+	{
+		FProcMeshSection *ProcSection =
+			ProcMeshComp->GetProcMeshSection(SectionIdx);
+		UMaterialInterface *Material = ProcMeshComp->GetMaterial(SectionIdx);
+		FPolygonGroupID *PolygonGroupID = UniqueMaterials.Find(Material);
+		check(PolygonGroupID != nullptr);
+		PolygonGroupForSection.Add(*PolygonGroupID);
+	}
+
+
+	// Add Vertex and VertexInstance and polygon for each section
+	for (int32 SectionIdx = 0; SectionIdx < NumSections; SectionIdx++)
+	{
+		FProcMeshSection *ProcSection =
+			ProcMeshComp->GetProcMeshSection(SectionIdx);
+		FPolygonGroupID PolygonGroupID = PolygonGroupForSection[SectionIdx];
+		// Create the vertex
+		int32 NumVertex = ProcSection->ProcVertexBuffer.Num();
+		TMap<int32, FVertexID> VertexIndexToVertexID;
+		VertexIndexToVertexID.Reserve(NumVertex);
+		for (int32 VertexIndex = 0; VertexIndex < NumVertex; ++VertexIndex)
+		{
+			FProcMeshVertex &Vert = ProcSection->ProcVertexBuffer[VertexIndex];
+			const FVertexID VertexID = MeshDescription.CreateVertex();
+			VertexPositions[VertexID] = Vert.Position;
+			VertexIndexToVertexID.Add(VertexIndex, VertexID);
+		}
+		// Create the VertexInstance
+		int32 NumIndices = ProcSection->ProcIndexBuffer.Num();
+		int32 NumTri = NumIndices / 3;
+		TMap<int32, FVertexInstanceID> IndiceIndexToVertexInstanceID;
+		IndiceIndexToVertexInstanceID.Reserve(NumVertex);
+		for (int32 IndiceIndex = 0; IndiceIndex < NumIndices; IndiceIndex++)
+		{
+			const int32 VertexIndex = ProcSection->ProcIndexBuffer[IndiceIndex];
+			const FVertexID VertexID = VertexIndexToVertexID[VertexIndex];
+			const FVertexInstanceID VertexInstanceID =
+				MeshDescription.CreateVertexInstance(VertexID);
+			IndiceIndexToVertexInstanceID.Add(IndiceIndex, VertexInstanceID);
+
+			FProcMeshVertex &ProcVertex = ProcSection->ProcVertexBuffer[VertexIndex];
+
+			Tangents[VertexInstanceID] = ProcVertex.Tangent.TangentX;
+			Normals[VertexInstanceID] = ProcVertex.Normal;
+			BinormalSigns[VertexInstanceID] =
+				ProcVertex.Tangent.bFlipTangentY ? -1.f : 1.f;
+
+			Colors[VertexInstanceID] = FLinearColor(ProcVertex.Color);
+
+			UVs.Set(VertexInstanceID, 0, ProcVertex.UV0);
+			UVs.Set(VertexInstanceID, 1, ProcVertex.UV1);
+			UVs.Set(VertexInstanceID, 2, ProcVertex.UV2);
+			UVs.Set(VertexInstanceID, 3, ProcVertex.UV3);
+		}
+
+		// Create the polygons for this section
+		for (int32 TriIdx = 0; TriIdx < NumTri; TriIdx++)
+		{
+			FVertexID VertexIndexes[3];
+			TArray<FVertexInstanceID> VertexInstanceIDs;
+			VertexInstanceIDs.SetNum(3);
+
+			for (int32 CornerIndex = 0; CornerIndex < 3; ++CornerIndex)
+			{
+				const int32 IndiceIndex = (TriIdx * 3) + CornerIndex;
+				const int32 VertexIndex = ProcSection->ProcIndexBuffer[IndiceIndex];
+				VertexIndexes[CornerIndex] = VertexIndexToVertexID[VertexIndex];
+				VertexInstanceIDs[CornerIndex] =
+					IndiceIndexToVertexInstanceID[IndiceIndex];
+			}
+
+			// Insert a polygon into the mesh
+			MeshDescription.CreatePolygon(PolygonGroupID, VertexInstanceIDs);
+		}
+	}
+	return MeshDescription;
+}
diff --git a/Engine/Plugins/Runtime/ProceduralMeshComponent/Source/ProceduralMeshComponent/ProceduralMeshComponent.Build.cs b/Engine/Plugins/Runtime/ProceduralMeshComponent/Source/ProceduralMeshComponent/ProceduralMeshComponent.Build.cs
--- a/Engine/Plugins/Runtime/ProceduralMeshComponent/Source/ProceduralMeshComponent/ProceduralMeshComponent.Build.cs
+++ b/Engine/Plugins/Runtime/ProceduralMeshComponent/Source/ProceduralMeshComponent/ProceduralMeshComponent.Build.cs
@@ -4,18 +4,20 @@ namespace UnrealBuildTool.Rules
 {
 	public class ProceduralMeshComponent : ModuleRules
 	{
-        public ProceduralMeshComponent(ReadOnlyTargetRules Target) : base(Target)
+		public ProceduralMeshComponent(ReadOnlyTargetRules Target) : base(Target)
 		{
-            PrivateIncludePaths.Add("../../../../Shaders/Shared");
+			PrivateIncludePaths.Add("../../../../Shaders/Shared");
 
-            PublicDependencyModuleNames.AddRange(
+			PublicDependencyModuleNames.AddRange(
 				new string[]
 				{
 					"Core",
 					"CoreUObject",
-                    "Engine",
-                    "RenderCore",
-                    "RHI"
+					"Engine",
+					"InteractiveToolsFramework",
+					"MeshDescription",
+					"RenderCore",
+					"RHI"
 				}
 				);
 		}
diff --git a/Engine/Plugins/Runtime/ProceduralMeshComponent/Source/ProceduralMeshComponent/Public/ProceduralMeshBridge.h b/Engine/Plugins/Runtime/ProceduralMeshComponent/Source/ProceduralMeshComponent/Public/ProceduralMeshBridge.h
new file mode 100644
--- /dev/null
+++ b/Engine/Plugins/Runtime/ProceduralMeshComponent/Source/ProceduralMeshComponent/Public/ProceduralMeshBridge.h
@@ -0,0 +1,26 @@
+// Copyright 1998-2019 Epic Games, Inc. All Rights Reserved.
+
+#pragma once
+
+#include "CoreMinimal.h"
+#include "ComponentSourceInterfaces.h"
+#include "MeshDescription.h"
+#include "ProceduralMeshcomponent.h"
+
+class PROCEDURALMESHCOMPONENT_API  FProceduralMeshComponentTargetFactory : public FComponentTargetFactory
+{
+public:
+	bool CanBuild( UActorComponent* Candidate ) override;
+	TUniquePtr<FPrimitiveComponentTarget> Build( UPrimitiveComponent* PrimitiveComponent ) override;
+};
+
+class PROCEDURALMESHCOMPONENT_API FProceduralMeshComponentTarget : public FPrimitiveComponentTarget
+{
+public:
+	FProceduralMeshComponentTarget( UPrimitiveComponent* Component );
+
+	FMeshDescription* GetMesh() override;
+	void CommitMesh( const FCommitter& ) override;
+private:
+	TUniquePtr<FMeshDescription> MeshDescription;
+};
diff --git a/Engine/Plugins/Runtime/ProceduralMeshComponent/Source/ProceduralMeshComponent/Public/ProceduralMeshConversion.h b/Engine/Plugins/Runtime/ProceduralMeshComponent/Source/ProceduralMeshComponent/Public/ProceduralMeshConversion.h
new file mode 100644
--- /dev/null
+++ b/Engine/Plugins/Runtime/ProceduralMeshComponent/Source/ProceduralMeshComponent/Public/ProceduralMeshConversion.h
@@ -0,0 +1,16 @@
+// Copyright 1998-2019 Epic Games, Inc. All Rights Reserved.
+
+#pragma once
+
+#include "CoreMinimal.h"
+#include "MeshDescription.h"
+
+class UProceduralMeshComponent;
+
+FMeshDescription
+PROCEDURALMESHCOMPONENT_API
+BuildMeshDescription( UProceduralMeshComponent* ProcMeshComp );
+
+void
+PROCEDURALMESHCOMPONENT_API
+MeshDescriptionToProcMesh( const FMeshDescription& MeshDescription, UProceduralMeshComponent* ProcMeshComp );
diff --git a/Engine/Plugins/Runtime/ProceduralMeshComponent/Source/ProceduralMeshComponentEditor/Private/ProceduralMeshComponentDetails.cpp b/Engine/Plugins/Runtime/ProceduralMeshComponent/Source/ProceduralMeshComponentEditor/Private/ProceduralMeshComponentDetails.cpp
--- a/Engine/Plugins/Runtime/ProceduralMeshComponent/Source/ProceduralMeshComponentEditor/Private/ProceduralMeshComponentDetails.cpp
+++ b/Engine/Plugins/Runtime/ProceduralMeshComponent/Source/ProceduralMeshComponentEditor/Private/ProceduralMeshComponentDetails.cpp
@@ -1,6 +1,7 @@
 // Copyright 1998-2019 Epic Games, Inc. All Rights Reserved.
 
 #include "ProceduralMeshComponentDetails.h"
+#include "ProceduralMeshConversion.h"
 #include "Modules/ModuleManager.h"
 #include "Misc/PackageName.h"
 #include "Widgets/SNullWidget.h"
@@ -81,13 +82,11 @@ UProceduralMeshComponent* FProceduralMeshComponentDetails::GetFirstSelectedProcM
 	return ProcMeshComp;
 }
 
-
 bool FProceduralMeshComponentDetails::ConvertToStaticMeshEnabled() const
 {
 	return GetFirstSelectedProcMeshComp() != nullptr;
 }
 
-
 FReply FProceduralMeshComponentDetails::ClickedOnConvertToStaticMesh()
 {
 	// Find first selected ProcMeshComp
@@ -120,112 +119,7 @@ FReply FProceduralMeshComponentDetails::ClickedOnConvertToStaticMesh()
 			}
 
 
-			FMeshDescription MeshDescription;
-			UStaticMesh::RegisterMeshAttributes(MeshDescription);
-			FStaticMeshDescriptionAttributeGetter AttributeGetter(&MeshDescription);
-			TPolygonGroupAttributesRef<FName> PolygonGroupNames = AttributeGetter.GetPolygonGroupImportedMaterialSlotNames();
-			TVertexAttributesRef<FVector> VertexPositions = AttributeGetter.GetPositions();
-			TVertexInstanceAttributesRef<FVector> Tangents = AttributeGetter.GetTangents();
-			TVertexInstanceAttributesRef<float> BinormalSigns = AttributeGetter.GetBinormalSigns();
-			TVertexInstanceAttributesRef<FVector> Normals = AttributeGetter.GetNormals();
-			TVertexInstanceAttributesRef<FVector4> Colors = AttributeGetter.GetColors();
-			TVertexInstanceAttributesRef<FVector2D> UVs = AttributeGetter.GetUVs();
-			TEdgeAttributesRef<bool> EdgeHardnesses = AttributeGetter.GetEdgeHardnesses();
-			TEdgeAttributesRef<float> EdgeCreaseSharpnesses = AttributeGetter.GetEdgeCreaseSharpnesses();
-
-			// Materials to apply to new mesh
-			const int32 NumSections = ProcMeshComp->GetNumSections();
-			int32 VertexCount = 0;
-			int32 VertexInstanceCount = 0;
-			int32 PolygonCount = 0;
-			TMap<UMaterialInterface*, FPolygonGroupID> UniqueMaterials;
-			UniqueMaterials.Reserve(NumSections);
-			TArray<FPolygonGroupID> MaterialRemap;
-			MaterialRemap.Reserve(NumSections);
-			//Get all the info we need to create the MeshDescription
-			for (int32 SectionIdx = 0; SectionIdx < NumSections; SectionIdx++)
-			{
-				FProcMeshSection* ProcSection = ProcMeshComp->GetProcMeshSection(SectionIdx);
-				VertexCount += ProcSection->ProcVertexBuffer.Num();
-				VertexInstanceCount += ProcSection->ProcIndexBuffer.Num();
-				PolygonCount += ProcSection->ProcIndexBuffer.Num() / 3;
-				UMaterialInterface*Material = ProcMeshComp->GetMaterial(SectionIdx);
-				if (!UniqueMaterials.Contains(Material))
-				{
-					FPolygonGroupID NewPolygonGroup = MeshDescription.CreatePolygonGroup();
-					UniqueMaterials.Add(Material, NewPolygonGroup);
-					PolygonGroupNames[NewPolygonGroup] = Material->GetFName();
-				}
-				FPolygonGroupID* PolygonGroupID = UniqueMaterials.Find(Material);
-				check(PolygonGroupID != nullptr);
-				MaterialRemap.Add(*PolygonGroupID);
-			}
-			MeshDescription.ReserveNewVertices(VertexCount);
-			MeshDescription.ReserveNewVertexInstances(VertexInstanceCount);
-			MeshDescription.ReserveNewPolygons(PolygonCount);
-			MeshDescription.ReserveNewEdges(PolygonCount * 2);
-			UVs.SetNumIndices(4);
-			//Add Vertex and VertexInstance and polygon for each section
-			for (int32 SectionIdx = 0; SectionIdx < NumSections; SectionIdx++)
-			{
-				FProcMeshSection* ProcSection = ProcMeshComp->GetProcMeshSection(SectionIdx);
-				FPolygonGroupID PolygonGroupID = MaterialRemap[SectionIdx];
-				//Create the vertex
-				int32 NumVertex = ProcSection->ProcVertexBuffer.Num();
-				TMap<int32, FVertexID> VertexIndexToVertexID;
-				VertexIndexToVertexID.Reserve(NumVertex); 
-				for (int32 VertexIndex = 0; VertexIndex < NumVertex; ++VertexIndex)
-				{
-					FProcMeshVertex& Vert = ProcSection->ProcVertexBuffer[VertexIndex];
-					const FVertexID VertexID = MeshDescription.CreateVertex();
-					VertexPositions[VertexID] = Vert.Position;
-					VertexIndexToVertexID.Add(VertexIndex, VertexID);
-				}
-				//Create the VertexInstance
-				int32 NumIndices = ProcSection->ProcIndexBuffer.Num();
-				int32 NumTri = NumIndices / 3;
-				TMap<int32, FVertexInstanceID> IndiceIndexToVertexInstanceID;
-				IndiceIndexToVertexInstanceID.Reserve(NumVertex);
-				for (int32 IndiceIndex = 0; IndiceIndex < NumIndices; IndiceIndex++)
-				{
-					const int32 VertexIndex = ProcSection->ProcIndexBuffer[IndiceIndex];
-					const FVertexID VertexID = VertexIndexToVertexID[VertexIndex];
-					const FVertexInstanceID VertexInstanceID = MeshDescription.CreateVertexInstance(VertexID);
-					IndiceIndexToVertexInstanceID.Add(IndiceIndex, VertexInstanceID);
-
-					FProcMeshVertex& ProcVertex = ProcSection->ProcVertexBuffer[VertexIndex];
-
-					Tangents[VertexInstanceID] = ProcVertex.Tangent.TangentX;
-					Normals[VertexInstanceID] = ProcVertex.Normal;
-					BinormalSigns[VertexInstanceID] = ProcVertex.Tangent.bFlipTangentY ? -1.f : 1.f;
-
-					Colors[VertexInstanceID] = FLinearColor(ProcVertex.Color);
-
-					UVs.Set(VertexInstanceID, 0, ProcVertex.UV0);
-					UVs.Set(VertexInstanceID, 1, ProcVertex.UV1);
-					UVs.Set(VertexInstanceID, 2, ProcVertex.UV2);
-					UVs.Set(VertexInstanceID, 3, ProcVertex.UV3);
-				}
-				
-				//Create the polygons for this section
-				for (int32 TriIdx = 0; TriIdx < NumTri; TriIdx++)
-				{
-					FVertexID VertexIndexes[3];
-					TArray<FVertexInstanceID> VertexInstanceIDs;
-					VertexInstanceIDs.SetNum(3);
-
-					for (int32 CornerIndex = 0; CornerIndex < 3; ++CornerIndex)
-					{
-						const int32 IndiceIndex = (TriIdx * 3) + CornerIndex;
-						const int32 VertexIndex = ProcSection->ProcIndexBuffer[IndiceIndex];
-						VertexIndexes[CornerIndex] = VertexIndexToVertexID[VertexIndex];
-						VertexInstanceIDs[CornerIndex] = IndiceIndexToVertexInstanceID[IndiceIndex];
-					}
-
-					// Insert a polygon into the mesh
-					MeshDescription.CreatePolygon(PolygonGroupID, VertexInstanceIDs);
-				}
-			}
+			FMeshDescription MeshDescription = BuildMeshDescription(ProcMeshComp);
 
 			// If we got some valid data.
 			if (MeshDescription.Polygons().Num() > 0)
@@ -258,10 +152,18 @@ FReply FProceduralMeshComponentDetails::ClickedOnConvertToStaticMesh()
 				*OriginalMeshDescription = MeshDescription;
 				StaticMesh->CommitMeshDescription(0);
 
+				TSet<UMaterialInterface*> UniqueMaterials;
+				const int32 NumSections = ProcMeshComp->GetNumSections();
+				for (int32 SectionIdx = 0; SectionIdx < NumSections; SectionIdx++)
+				{
+					FProcMeshSection *ProcSection =
+						ProcMeshComp->GetProcMeshSection(SectionIdx);
+					UMaterialInterface *Material = ProcMeshComp->GetMaterial(SectionIdx);
+					UniqueMaterials.Add(Material);
+				}
 				// Copy materials to new mesh
-				for (auto Kvp : UniqueMaterials)
+				for (auto* Material : UniqueMaterials)
 				{
-					UMaterialInterface* Material = Kvp.Key;
 					StaticMesh->StaticMaterials.Add(FStaticMaterial(Material, Material->GetFName(), Material->GetFName()));
 				}
 
diff --git a/Engine/Source/Editor/Experimental/EditorInteractiveToolsFramework/Private/EditorComponentSourceFactory.cpp b/Engine/Source/Editor/Experimental/EditorInteractiveToolsFramework/Private/EditorComponentSourceFactory.cpp
--- a/Engine/Source/Editor/Experimental/EditorInteractiveToolsFramework/Private/EditorComponentSourceFactory.cpp
+++ b/Engine/Source/Editor/Experimental/EditorInteractiveToolsFramework/Private/EditorComponentSourceFactory.cpp
@@ -8,19 +8,17 @@
 #include "PhysicsEngine/BodySetup.h"
 
 
-namespace {
 FMeshDescription *
-GetMeshDescription(UStaticMeshComponent* Component, int LODIndex) {
-	return Component->GetStaticMesh()->GetMeshDescription(LODIndex);
+FStaticMeshComponentTarget::GetMesh() {
+	return Cast<UStaticMeshComponent>(Component)->GetStaticMesh()->GetMeshDescription(LODIndex);
 }
 
 void
-CommitInPlaceModification(UStaticMeshComponent* Component, int LODIndex,
-						  const TFunction<void(FMeshDescription*)>& ModifyFunction)
+FStaticMeshComponentTarget::CommitMesh( const FCommitter& Committer )
 {
 	//bool bSaved = Component->Modify();
 	//check(bSaved);
-	UStaticMesh* StaticMesh = Component->GetStaticMesh();
+	UStaticMesh* StaticMesh = Cast<UStaticMeshComponent>(Component)->GetStaticMesh();
 
 	// make sure transactional flag is on
 	StaticMesh->SetFlags(RF_Transactional);
@@ -29,7 +27,7 @@ CommitInPlaceModification(UStaticMeshComponent* Component, int LODIndex,
 	check(bSavedToTransactionBuffer);
 	FMeshDescription* MeshDescription = StaticMesh->GetMeshDescription(LODIndex);
 
-	ModifyFunction(MeshDescription);
+	Committer(MeshDescription);
 
 	StaticMesh->CommitMeshDescription(LODIndex);
 	StaticMesh->PostEditChange();
@@ -39,25 +37,20 @@ CommitInPlaceModification(UStaticMeshComponent* Component, int LODIndex,
 
 	//Component->PostEditChange();
 }
-} // namespace
 
-FMeshDescriptionBridge
-MakeStaticMeshDescriptionBridge(UPrimitiveComponent* Component)
+bool
+FStaticMeshComponentTargetFactory::CanBuild(UActorComponent* Component)
+{
+	return Component != nullptr;
+}
+
+TUniquePtr< FPrimitiveComponentTarget >
+FStaticMeshComponentTargetFactory::Build(UPrimitiveComponent* Component)
 {
-	const int LODIndex{0};
 	UStaticMeshComponent* StaticMeshComponent = Cast<UStaticMeshComponent>(Component);
 	if (StaticMeshComponent != nullptr)
 	{
-		return {
-			[StaticMeshComponent, LODIndex]()
-			{
-				return GetMeshDescription(StaticMeshComponent, LODIndex);
-			},
-			[StaticMeshComponent, LODIndex](const FMeshDescriptionBridge::FCommitter& Committer)
-			{
-				CommitInPlaceModification(StaticMeshComponent, LODIndex, Committer);
-			}
-		};
+		return TUniquePtr< FPrimitiveComponentTarget > { new FStaticMeshComponentTarget{Component} };
 	}
-	return FMeshDescriptionBridge{};
+	return {};
 }
diff --git a/Engine/Source/Editor/Experimental/EditorInteractiveToolsFramework/Private/EditorInteractiveToolsFrameworkModule.cpp b/Engine/Source/Editor/Experimental/EditorInteractiveToolsFramework/Private/EditorInteractiveToolsFrameworkModule.cpp
--- a/Engine/Source/Editor/Experimental/EditorInteractiveToolsFramework/Private/EditorInteractiveToolsFrameworkModule.cpp
+++ b/Engine/Source/Editor/Experimental/EditorInteractiveToolsFramework/Private/EditorInteractiveToolsFrameworkModule.cpp
@@ -9,8 +9,7 @@
 void FEditorInteractiveToolsFrameworkModule::StartupModule()
 {
 	// This code will execute after your module is loaded into memory; the exact timing is specified in the .uplugin file per-module
-	AddMeshDescriptionBridgeBuilder(MakeStaticMeshDescriptionBridge);
-	UE_LOG(LogTemp, Warning, TEXT("FEditorInteractiveToolsFrameworkModule StartupModule"));
+	AddComponentTargetFactory( TUniquePtr<FComponentTargetFactory>{new FStaticMeshComponentTargetFactory{} } );
 }
 
 void FEditorInteractiveToolsFrameworkModule::ShutdownModule()
diff --git a/Engine/Source/Editor/Experimental/EditorInteractiveToolsFramework/Public/EditorComponentSourceFactory.h b/Engine/Source/Editor/Experimental/EditorInteractiveToolsFramework/Public/EditorComponentSourceFactory.h
--- a/Engine/Source/Editor/Experimental/EditorInteractiveToolsFramework/Public/EditorComponentSourceFactory.h
+++ b/Engine/Source/Editor/Experimental/EditorInteractiveToolsFramework/Public/EditorComponentSourceFactory.h
@@ -4,5 +4,22 @@
 
 #include "CoreMinimal.h"
 #include "ComponentSourceInterfaces.h"
+#include "Components/PrimitiveComponent.h"
+#include "Components/StaticMeshComponent.h"
 
-EDITORINTERACTIVETOOLSFRAMEWORK_API FMeshDescriptionBridge MakeStaticMeshDescriptionBridge(UPrimitiveComponent* Component);
+class EDITORINTERACTIVETOOLSFRAMEWORK_API FStaticMeshComponentTargetFactory : public FComponentTargetFactory
+{
+public:
+	bool CanBuild( UActorComponent* Candidate ) override;
+	TUniquePtr<FPrimitiveComponentTarget> Build( UPrimitiveComponent* PrimitiveComponent ) override;
+};
+
+class EDITORINTERACTIVETOOLSFRAMEWORK_API FStaticMeshComponentTarget : public FPrimitiveComponentTarget
+{
+public:
+	FStaticMeshComponentTarget( UPrimitiveComponent* Component )
+		: FPrimitiveComponentTarget( Cast<UStaticMeshComponent>(Component) ){}
+	FMeshDescription* GetMesh() override;
+	void CommitMesh( const FCommitter& ) override;
+	static const int LODIndex{0};
+};
diff --git a/Engine/Source/Runtime/Engine/Classes/Components/MeshComponent.h b/Engine/Source/Runtime/Engine/Classes/Components/MeshComponent.h
--- a/Engine/Source/Runtime/Engine/Classes/Components/MeshComponent.h
+++ b/Engine/Source/Runtime/Engine/Classes/Components/MeshComponent.h
@@ -22,6 +22,7 @@ class ENGINE_API UMeshComponent : public UPrimitiveComponent
 {
 	GENERATED_UCLASS_BODY()
 
+public:
 	/** Per-Component material overrides.  These must NOT be set directly or a race condition can occur between GC and the rendering thread. */
 	UPROPERTY(EditAnywhere, AdvancedDisplay, Category=Rendering, Meta=(ToolTip="Material overrides."))
 	TArray<class UMaterialInterface*> OverrideMaterials;
diff --git a/Engine/Source/Runtime/Experimental/InteractiveToolsFramework/Private/BaseTools/MeshSurfacePointTool.cpp b/Engine/Source/Runtime/Experimental/InteractiveToolsFramework/Private/BaseTools/MeshSurfacePointTool.cpp
--- a/Engine/Source/Runtime/Experimental/InteractiveToolsFramework/Private/BaseTools/MeshSurfacePointTool.cpp
+++ b/Engine/Source/Runtime/Experimental/InteractiveToolsFramework/Private/BaseTools/MeshSurfacePointTool.cpp
@@ -12,7 +12,7 @@
 
 bool UMeshSurfacePointToolBuilder::CanBuildTool(const FToolBuilderState& SceneState) const
 {
-	return ToolBuilderUtil::CountComponents(SceneState, ToolBuilderUtil::IsMeshDescriptionSourceComponent) == 1;
+	return ToolBuilderUtil::CountComponents(SceneState, CanMakeComponentTarget) == 1;
 }
 
 UInteractiveTool* UMeshSurfacePointToolBuilder::BuildTool(const FToolBuilderState& SceneState) const
@@ -29,7 +29,7 @@ UMeshSurfacePointTool* UMeshSurfacePointToolBuilder::CreateNewTool(const FToolBu
 
 void UMeshSurfacePointToolBuilder::InitializeNewTool(UMeshSurfacePointTool* NewTool, const FToolBuilderState& SceneState) const
 {
-	UActorComponent* ActorComponent = ToolBuilderUtil::FindFirstComponent(SceneState, ToolBuilderUtil::IsMeshDescriptionSourceComponent);
+	UActorComponent* ActorComponent = ToolBuilderUtil::FindFirstComponent(SceneState, CanMakeComponentTarget);
 	UPrimitiveComponent* MeshComponent = Cast<UPrimitiveComponent>(ActorComponent);
 	check(MeshComponent != nullptr);
 	NewTool->SetSelection( MakeComponentTarget(MeshComponent) );
@@ -58,7 +58,7 @@ void UMeshSurfacePointTool::Setup()
 
 bool UMeshSurfacePointTool::HitTest(const FRay& Ray, FHitResult& OutHit)
 {
-	return ComponentTarget.HitTest(Ray, OutHit);
+	return ComponentTarget->HitTest(Ray, OutHit);
 }
 
 
diff --git a/Engine/Source/Runtime/Experimental/InteractiveToolsFramework/Private/ComponentSourceInterfaces.cpp b/Engine/Source/Runtime/Experimental/InteractiveToolsFramework/Private/ComponentSourceInterfaces.cpp
--- a/Engine/Source/Runtime/Experimental/InteractiveToolsFramework/Private/ComponentSourceInterfaces.cpp
+++ b/Engine/Source/Runtime/Experimental/InteractiveToolsFramework/Private/ComponentSourceInterfaces.cpp
@@ -7,63 +7,75 @@
 
 namespace
 {
-	TArray<FMeshDescriptionBridge::FBuilder> Builders;
+	TArray<TUniquePtr<FComponentTargetFactory>> Factories;
 }
 
 void
-AddMeshDescriptionBridgeBuilder( FMeshDescriptionBridge::FBuilder Builder )
+AddComponentTargetFactory( TUniquePtr<FComponentTargetFactory> Factory )
 {
-	Builders.Push(MoveTemp(Builder));
+	Factories.Push( MoveTemp(Factory) );
 }
 
-FComponentTarget
+bool
+CanMakeComponentTarget(UActorComponent* Component)
+{
+	for ( const auto& Factory : Factories )
+	{
+		if ( Factory->CanBuild(Component) )
+		{
+			return true;
+		}
+	}
+	return false;
+}
+
+TUniquePtr<FPrimitiveComponentTarget>
 MakeComponentTarget(UPrimitiveComponent* Component)
 {
-	for ( const auto& Builder : Builders )
+	for ( const auto& Factory : Factories )
 	{
-		auto Bridge = Builder( Component );
-		if ( Bridge.HasSource() )
+		if ( Factory->CanBuild( Component ) )
 		{
-			return {Bridge, Component};
+			return Factory->Build( Component );
 		}
 	}
 	return {};
 }
 
 AActor*
-FComponentTarget::GetOwnerActor() const
+FPrimitiveComponentTarget::GetOwnerActor() const
 {
 	return Component->GetOwner();
 }
 
 UPrimitiveComponent*
-FComponentTarget::GetOwnerComponent() const
+FPrimitiveComponentTarget::GetOwnerComponent() const
 {
 	return Component;
 }
 
 
 void
-FComponentTarget::SetOwnerVisibility(bool bVisible) const
+FPrimitiveComponentTarget::SetOwnerVisibility(bool bVisible) const
 {
 	Component->SetVisibility(bVisible);
 }
 
 UMaterialInterface*
-FComponentTarget::GetMaterial(int32 MaterialIndex) const
+FPrimitiveComponentTarget::GetMaterial(int32 MaterialIndex) const
 {
 	return Component->GetMaterial(MaterialIndex);
 }
 
 FTransform
-FComponentTarget::GetWorldTransform() const
+FPrimitiveComponentTarget::GetWorldTransform() const
 {
 	//return Component->GetOwner()->GetActorTransform();
 	return Component->GetComponentTransform();
 }
 
 bool
-FComponentTarget::HitTest(const FRay& WorldRay, FHitResult& OutHit) const
+FPrimitiveComponentTarget::HitTest(const FRay& WorldRay, FHitResult& OutHit) const
 {
 	FVector End = WorldRay.PointAt(HALF_WORLD_MAX);
 	if (Component->LineTraceComponent(OutHit, WorldRay.Origin, End, FCollisionQueryParams(SCENE_QUERY_STAT(HitTest), true)))
diff --git a/Engine/Source/Runtime/Experimental/InteractiveToolsFramework/Private/ToolBuilderUtil.cpp b/Engine/Source/Runtime/Experimental/InteractiveToolsFramework/Private/ToolBuilderUtil.cpp
--- a/Engine/Source/Runtime/Experimental/InteractiveToolsFramework/Private/ToolBuilderUtil.cpp
+++ b/Engine/Source/Runtime/Experimental/InteractiveToolsFramework/Private/ToolBuilderUtil.cpp
@@ -8,18 +8,6 @@
 
 
 
-bool ToolBuilderUtil::IsMeshDescriptionSourceComponent(UActorComponent* ComponentObject)
-{
-	UStaticMeshComponent* StaticMeshComp = Cast<UStaticMeshComponent>(ComponentObject);
-	if (StaticMeshComp != nullptr)
-	{
-		return true;
-	}
-	return false;
-}
-
-
-
 int ToolBuilderUtil::CountComponents(const FToolBuilderState& InputState, const TFunction<bool(UActorComponent*)>& Predicate)
 {
 	int nTypedComponents = 0;
diff --git a/Engine/Source/Runtime/Experimental/InteractiveToolsFramework/Public/ComponentSourceInterfaces.h b/Engine/Source/Runtime/Experimental/InteractiveToolsFramework/Public/ComponentSourceInterfaces.h
--- a/Engine/Source/Runtime/Experimental/InteractiveToolsFramework/Public/ComponentSourceInterfaces.h
+++ b/Engine/Source/Runtime/Experimental/InteractiveToolsFramework/Public/ComponentSourceInterfaces.h
@@ -10,34 +10,11 @@
 
 // predeclarations
 class AActor;
+class UActorComponent;
 class UPrimitiveComponent;
 struct FMeshDescription;
 class UMaterialInterface;
 
-template <class MeshT>
-struct TMeshBridge
-{
-	using  FSource    = TFunction< MeshT*() >;
-	using  FCommitter = TFunction< void( MeshT* ) >;
-	using  FSink      = TFunction< void( const FCommitter& ) >;
-	using  FBuilder   = TFunction< TMeshBridge ( UPrimitiveComponent* ) >;
-
-	bool HasSource(){ return !!GetMesh; }
-	bool HasSink(){ return !!CommitMesh; }
-
-	FSource GetMesh{nullptr};
-	FSink   CommitMesh{nullptr};
-};
-using FMeshDescriptionBridge = TMeshBridge<FMeshDescription>;
-
-/**
- * Add a factory method to make MeshDescriptionSources from UActorComponent*
- * @param Builder The MeshDescriptionSourceBuilder
- * @return void
- */
-INTERACTIVETOOLSFRAMEWORK_API void AddMeshDescriptionBridgeBuilder( FMeshDescriptionBridge::FBuilder Builder );
-
-
 /**
  * Wrapper around a UObject Component that can provide a MeshDescription, and
  * (optionally) bake a modified MeshDescription back to this Component.
@@ -46,9 +23,16 @@ INTERACTIVETOOLSFRAMEWORK_API void AddMeshDescriptionBridgeBuilder( FMeshDescrip
  *
  * (Conceivably this doesn't have to be backed by a Component, but most usage will assume there is an Actor)
  */
-class INTERACTIVETOOLSFRAMEWORK_API FComponentTarget
+class INTERACTIVETOOLSFRAMEWORK_API FPrimitiveComponentTarget
 {
 public:
+	virtual ~FPrimitiveComponentTarget(){}
+
+	/** Constructor UPrimitivecomponent*
+	 *  @param Component the UPrimitiveComponent* to target
+	 */
+	FPrimitiveComponentTarget( UPrimitiveComponent* Component ): Component( Component ){}
+
 	/** @return the Actor that owns this Component */
 	AActor* GetOwnerActor() const;
 
@@ -82,10 +66,29 @@ public:
 	 */
 	void SetOwnerVisibility(bool bVisible) const;
 
-	FMeshDescriptionBridge MeshDescriptionBridge{};
+	using  FCommitter  = TFunction< void( FMeshDescription* ) >;
+	virtual FMeshDescription* GetMesh() = 0;
+	virtual void CommitMesh( const FCommitter& ) = 0;
+
 	UPrimitiveComponent* Component{};
 };
 
+class INTERACTIVETOOLSFRAMEWORK_API FComponentTargetFactory
+{
+public:
+	virtual ~FComponentTargetFactory(){}
+	virtual bool CanBuild( UActorComponent* Candidate ) = 0;
+	virtual TUniquePtr<FPrimitiveComponentTarget> Build( UPrimitiveComponent* PrimitiveComponent ) = 0;
+};
+
+/**
+ * Add a factory method to make ComponentTarget from UPrimitiveComponent*
+ * @param Factory The ComponentTargetFactory
+ * @return void
+ */
+INTERACTIVETOOLSFRAMEWORK_API void AddComponentTargetFactory( TUniquePtr<FComponentTargetFactory> Factory );
+
+
 /**
  * Create a TargetComponent for the given Component
  * @param Component A UObject that we would like to use as tool target. This must presently descend from
@@ -93,4 +96,12 @@ public:
  * @return An FComponentTarget instance. Must not return null, though the MeshSource and MeshSink in it's MeshBridge may
  * be
  */
-INTERACTIVETOOLSFRAMEWORK_API FComponentTarget MakeComponentTarget(UPrimitiveComponent* Component);
+INTERACTIVETOOLSFRAMEWORK_API TUniquePtr<FPrimitiveComponentTarget> MakeComponentTarget(UPrimitiveComponent* Component);
+
+
+/**
+ * Determine whether a TargetComponent can be created for the given Component
+ * @param Component A UObject that we would like to use as tool target.
+ * @return bool signifying whether or not a ComponentTarget can be built
+ */
+INTERACTIVETOOLSFRAMEWORK_API bool CanMakeComponentTarget(UActorComponent* Component);
diff --git a/Engine/Source/Runtime/Experimental/InteractiveToolsFramework/Public/MultiSelectionTool.h b/Engine/Source/Runtime/Experimental/InteractiveToolsFramework/Public/MultiSelectionTool.h
--- a/Engine/Source/Runtime/Experimental/InteractiveToolsFramework/Public/MultiSelectionTool.h
+++ b/Engine/Source/Runtime/Experimental/InteractiveToolsFramework/Public/MultiSelectionTool.h
@@ -11,10 +11,10 @@ class INTERACTIVETOOLSFRAMEWORK_API UMultiSelectionTool : public UInteractiveToo
 {
 GENERATED_BODY()
 public:
-	void SetSelection(TArray<FComponentTarget> ComponentTargetsIn)
+	void SetSelection(TArray<TUniquePtr<FPrimitiveComponentTarget>> ComponentTargetsIn)
     {
 		ComponentTargets = MoveTemp(ComponentTargetsIn);
 	}
 protected:
-	TArray<FComponentTarget> ComponentTargets{};
+	TArray<TUniquePtr<FPrimitiveComponentTarget>> ComponentTargets{};
 };
diff --git a/Engine/Source/Runtime/Experimental/InteractiveToolsFramework/Public/SingleSelectionTool.h b/Engine/Source/Runtime/Experimental/InteractiveToolsFramework/Public/SingleSelectionTool.h
--- a/Engine/Source/Runtime/Experimental/InteractiveToolsFramework/Public/SingleSelectionTool.h
+++ b/Engine/Source/Runtime/Experimental/InteractiveToolsFramework/Public/SingleSelectionTool.h
@@ -11,10 +11,10 @@ class INTERACTIVETOOLSFRAMEWORK_API USingleSelectionTool : public UInteractiveTo
 {
 GENERATED_BODY()
 public:
-	void SetSelection(FComponentTarget ComponentTargetIn)
+	void SetSelection(TUniquePtr<FPrimitiveComponentTarget> ComponentTargetIn)
     {
 		ComponentTarget = MoveTemp(ComponentTargetIn);
 	}
 protected:
-	FComponentTarget ComponentTarget{};
+	TUniquePtr<FPrimitiveComponentTarget> ComponentTarget{};
 };
diff --git a/Engine/Source/Runtime/Experimental/InteractiveToolsFramework/Public/ToolBuilderUtil.h b/Engine/Source/Runtime/Experimental/InteractiveToolsFramework/Public/ToolBuilderUtil.h
--- a/Engine/Source/Runtime/Experimental/InteractiveToolsFramework/Public/ToolBuilderUtil.h
+++ b/Engine/Source/Runtime/Experimental/InteractiveToolsFramework/Public/ToolBuilderUtil.h
@@ -12,10 +12,6 @@
 */
 namespace ToolBuilderUtil
 {
-	/** Returns true if this UObject can provide a FMeshDescription */
-	INTERACTIVETOOLSFRAMEWORK_API
-	bool IsMeshDescriptionSourceComponent(UActorComponent* ComponentObject);
-	
 	/** Count number of selected components that pass predicate. If Component selection is not empty, returns that count, otherwise counts in all selected Actors */
 	INTERACTIVETOOLSFRAMEWORK_API
 	int CountComponents(const FToolBuilderState& InputState, const TFunction<bool(UActorComponent*)>& Predicate);
