Fix case on header include

#rnx

diff --git a/Engine/Plugins/Runtime/ProceduralMeshComponent/Source/ProceduralMeshComponent/Public/ProceduralMeshBridge.h b/Engine/Plugins/Runtime/ProceduralMeshComponent/Source/ProceduralMeshComponent/Public/ProceduralMeshBridge.h
--- a/Engine/Plugins/Runtime/ProceduralMeshComponent/Source/ProceduralMeshComponent/Public/ProceduralMeshBridge.h
+++ b/Engine/Plugins/Runtime/ProceduralMeshComponent/Source/ProceduralMeshComponent/Public/ProceduralMeshBridge.h
@@ -5,7 +5,7 @@
 #include "CoreMinimal.h"
 #include "ComponentSourceInterfaces.h"
 #include "MeshDescription.h"
-#include "ProceduralMeshcomponent.h"
+#include "ProceduralMeshComponent.h"
 
 class PROCEDURALMESHCOMPONENT_API  FProceduralMeshComponentTargetFactory : public FComponentTargetFactory
 {
