Copy simple collision data when generating a Static Mesh asset from a Procedural Mesh

diff --git a/Engine/Plugins/Runtime/ProceduralMeshComponent/Source/ProceduralMeshComponentEditor/Private/ProceduralMeshComponentDetails.cpp b/Engine/Plugins/Runtime/ProceduralMeshComponent/Source/ProceduralMeshComponentEditor/Private/ProceduralMeshComponentDetails.cpp
--- a/Engine/Plugins/Runtime/ProceduralMeshComponent/Source/ProceduralMeshComponentEditor/Private/ProceduralMeshComponentDetails.cpp
+++ b/Engine/Plugins/Runtime/ProceduralMeshComponent/Source/ProceduralMeshComponentEditor/Private/ProceduralMeshComponentDetails.cpp
@@ -144,6 +144,20 @@ FReply FProceduralMeshComponentDetails::ClickedOnConvertToStaticMesh()
 				StaticMesh->CreateMeshDescription(0, MoveTemp(MeshDescription));
 				StaticMesh->CommitMeshDescription(0);
 
+				//// SIMPLE COLLISION
+				if (!ProcMeshComp->bUseComplexAsSimpleCollision )
+				{
+					StaticMesh->CreateBodySetup();
+					UBodySetup* NewBodySetup = StaticMesh->BodySetup;
+					NewBodySetup->BodySetupGuid = FGuid::NewGuid();
+					NewBodySetup->AggGeom.ConvexElems = ProcMeshComp->ProcMeshBodySetup->AggGeom.ConvexElems;
+					NewBodySetup->bGenerateMirroredCollision = false;
+					NewBodySetup->bDoubleSidedGeometry = true;
+					NewBodySetup->CollisionTraceFlag = CTF_UseDefault;
+					NewBodySetup->CreatePhysicsMeshes();
+				}
+
+				//// MATERIALS
 				TSet<UMaterialInterface*> UniqueMaterials;
 				const int32 NumSections = ProcMeshComp->GetNumSections();
 				for (int32 SectionIdx = 0; SectionIdx < NumSections; SectionIdx++)
