Add #include for the strict build to fix CIS

#rnx

diff --git a/Engine/Plugins/Runtime/ProceduralMeshComponent/Source/ProceduralMeshComponentEditor/Private/ProceduralMeshComponentDetails.cpp b/Engine/Plugins/Runtime/ProceduralMeshComponent/Source/ProceduralMeshComponentEditor/Private/ProceduralMeshComponentDetails.cpp
--- a/Engine/Plugins/Runtime/ProceduralMeshComponent/Source/ProceduralMeshComponentEditor/Private/ProceduralMeshComponentDetails.cpp
+++ b/Engine/Plugins/Runtime/ProceduralMeshComponent/Source/ProceduralMeshComponentEditor/Private/ProceduralMeshComponentDetails.cpp
@@ -19,7 +19,7 @@
 #include "ProceduralMeshComponent.h"
 #include "StaticMeshAttributes.h"
 #include "MeshDescriptionOperations.h"
-
+#include "PhysicsEngine/BodySetup.h"
 #include "Dialogs/DlgPickAssetPath.h"
 #include "AssetRegistryModule.h"
 
