ModelingTools: remove old PrimitiveComponentTarget/ToolTarget support from ProceduralMeshComponent module, to remove InteractiveToolFramework dependency
#rnx

diff --git a/Engine/Plugins/Runtime/ProceduralMeshComponent/Source/ProceduralMeshComponent/Private/ProceduralMeshBridge.cpp b/Engine/Plugins/Runtime/ProceduralMeshComponent/Source/ProceduralMeshComponent/Private/ProceduralMeshBridge.cpp
deleted file mode 100644
diff --git a/Engine/Plugins/Runtime/ProceduralMeshComponent/Source/ProceduralMeshComponent/Private/ProceduralMeshComponentPlugin.cpp b/Engine/Plugins/Runtime/ProceduralMeshComponent/Source/ProceduralMeshComponent/Private/ProceduralMeshComponentPlugin.cpp
--- a/Engine/Plugins/Runtime/ProceduralMeshComponent/Source/ProceduralMeshComponent/Private/ProceduralMeshComponentPlugin.cpp
+++ b/Engine/Plugins/Runtime/ProceduralMeshComponent/Source/ProceduralMeshComponent/Private/ProceduralMeshComponentPlugin.cpp
@@ -1,9 +1,7 @@
 // Copyright Epic Games, Inc. All Rights Reserved.
 
 #include "CoreMinimal.h"
-#include "ComponentSourceInterfaces.h"
 #include "Modules/ModuleManager.h"
-#include "ProceduralMeshBridge.h"
 #include "IProceduralMeshComponentPlugin.h"
 
 class FProceduralMeshComponentPlugin : public IProceduralMeshComponentPlugin
@@ -20,7 +18,6 @@ IMPLEMENT_MODULE( FProceduralMeshComponentPlugin, ProceduralMeshComponent )
 void FProceduralMeshComponentPlugin::StartupModule()
 {
 	// This code will execute after your module is loaded into memory; the exact timing is specified in the .uplugin file per-module
-	AddComponentTargetFactory( TUniquePtr<FComponentTargetFactory>{new FProceduralMeshComponentTargetFactory{} } );
 }
 
 
diff --git a/Engine/Plugins/Runtime/ProceduralMeshComponent/Source/ProceduralMeshComponent/Private/ProceduralMeshComponentToolTarget.cpp b/Engine/Plugins/Runtime/ProceduralMeshComponent/Source/ProceduralMeshComponent/Private/ProceduralMeshComponentToolTarget.cpp
deleted file mode 100644
diff --git a/Engine/Plugins/Runtime/ProceduralMeshComponent/Source/ProceduralMeshComponent/ProceduralMeshComponent.Build.cs b/Engine/Plugins/Runtime/ProceduralMeshComponent/Source/ProceduralMeshComponent/ProceduralMeshComponent.Build.cs
--- a/Engine/Plugins/Runtime/ProceduralMeshComponent/Source/ProceduralMeshComponent/ProceduralMeshComponent.Build.cs
+++ b/Engine/Plugins/Runtime/ProceduralMeshComponent/Source/ProceduralMeshComponent/ProceduralMeshComponent.Build.cs
@@ -14,7 +14,6 @@ namespace UnrealBuildTool.Rules
 					"Core",
 					"CoreUObject",
 					"Engine",
-					"InteractiveToolsFramework",
 					"MeshDescription",
 					"RenderCore",
 					"RHI",
diff --git a/Engine/Plugins/Runtime/ProceduralMeshComponent/Source/ProceduralMeshComponent/Public/ProceduralMeshBridge.h b/Engine/Plugins/Runtime/ProceduralMeshComponent/Source/ProceduralMeshComponent/Public/ProceduralMeshBridge.h
deleted file mode 100644
diff --git a/Engine/Plugins/Runtime/ProceduralMeshComponent/Source/ProceduralMeshComponent/Public/ProceduralMeshComponentToolTarget.h b/Engine/Plugins/Runtime/ProceduralMeshComponent/Source/ProceduralMeshComponent/Public/ProceduralMeshComponentToolTarget.h
deleted file mode 100644
