ProceduralMeshComponent: Drop triangles with non-unique vertex indices.

diff --git a/Engine/Plugins/Runtime/ProceduralMeshComponent/Source/ProceduralMeshComponent/Private/ProceduralMeshComponent.cpp b/Engine/Plugins/Runtime/ProceduralMeshComponent/Source/ProceduralMeshComponent/Private/ProceduralMeshComponent.cpp
--- a/Engine/Plugins/Runtime/ProceduralMeshComponent/Source/ProceduralMeshComponent/Private/ProceduralMeshComponent.cpp
+++ b/Engine/Plugins/Runtime/ProceduralMeshComponent/Source/ProceduralMeshComponent/Private/ProceduralMeshComponent.cpp
@@ -575,16 +575,53 @@ void UProceduralMeshComponent::CreateMeshSection(int32 SectionIndex, const TArra
 		NewSection.SectionLocalBox += Vertex.Position;
 	}
 
-	// Copy index buffer (clamping to vertex range)
-	int32 NumTriIndices = Triangles.Num();
-	NumTriIndices = (NumTriIndices/3) * 3; // Ensure we have exact number of triangles (array is multiple of 3 long)
+	// Get triangle indices, clamping to vertex range
+	const int32 MaxIndex = NumVerts - 1;
+	const auto GetTriIndices = [&Triangles, MaxIndex](int32 Idx)
+	{
+		return TTuple<int32, int32, int32>(FMath::Min(Triangles[Idx    ], MaxIndex),
+										   FMath::Min(Triangles[Idx + 1], MaxIndex),
+			                               FMath::Min(Triangles[Idx + 2], MaxIndex));
+	};
+
+	const int32 NumTriIndices = (Triangles.Num() / 3) * 3; // Ensure number of triangle indices is multiple of three
+
+	// Detect degenerate triangles, i.e. non-unique vertex indices within the same triangle
+	int32 NumDegenerateTriangles = 0;
+	for (int32 IndexIdx = 0; IndexIdx < NumTriIndices; IndexIdx += 3)
+	{
+		int32 a, b, c;
+		Tie(a, b, c) = GetTriIndices(IndexIdx);
+		NumDegenerateTriangles += a == b || a == c || b == c;
+	}
+	if (NumDegenerateTriangles > 0)
+	{
+		UE_LOG(LogProceduralComponent, Warning, TEXT("Detected %d degenerate triangle%s with non-unique vertex indices for created mesh section in '%s'; degenerate triangles will be dropped."),
+			   NumDegenerateTriangles, NumDegenerateTriangles > 1 ? "s" : "", *GetFullName());
+	}
 
+	// Copy index buffer for non-degenerate triangles
 	NewSection.ProcIndexBuffer.Reset();
-	NewSection.ProcIndexBuffer.AddUninitialized(NumTriIndices);
-	for (int32 IndexIdx = 0; IndexIdx < NumTriIndices; IndexIdx++)
+	NewSection.ProcIndexBuffer.AddUninitialized(NumTriIndices - NumDegenerateTriangles * 3);
+	int32 CopyIndexIdx = 0;
+	for (int32 IndexIdx = 0; IndexIdx < NumTriIndices; IndexIdx += 3)
 	{
-		NewSection.ProcIndexBuffer[IndexIdx] = FMath::Min(Triangles[IndexIdx], NumVerts - 1);
+		int32 a, b, c;
+		Tie(a, b, c) = GetTriIndices(IndexIdx);
+
+		if (a != b && a != c && b != c)
+		{
+			NewSection.ProcIndexBuffer[CopyIndexIdx++] = a;
+			NewSection.ProcIndexBuffer[CopyIndexIdx++] = b;
+			NewSection.ProcIndexBuffer[CopyIndexIdx++] = c;
+		}
+		else
+		{
+			--NumDegenerateTriangles;
+		}
 	}
+	check(NumDegenerateTriangles == 0);
+	check(CopyIndexIdx == NewSection.ProcIndexBuffer.Num());
 
 	NewSection.bEnableCollision = bCreateCollision;
 
