Fixed MeshBatch.SegmentIndex assignment for procedural ray tracing meshes.

diff --git a/Engine/Plugins/Runtime/ProceduralMeshComponent/Source/ProceduralMeshComponent/Private/ProceduralMeshComponent.cpp b/Engine/Plugins/Runtime/ProceduralMeshComponent/Source/ProceduralMeshComponent/Private/ProceduralMeshComponent.cpp
--- a/Engine/Plugins/Runtime/ProceduralMeshComponent/Source/ProceduralMeshComponent/Private/ProceduralMeshComponent.cpp
+++ b/Engine/Plugins/Runtime/ProceduralMeshComponent/Source/ProceduralMeshComponent/Private/ProceduralMeshComponent.cpp
@@ -458,7 +458,7 @@ public:
 					FMeshBatch MeshBatch;
 
 					MeshBatch.VertexFactory = &Section->VertexFactory;
-					MeshBatch.SegmentIndex = SegmentIndex;
+					MeshBatch.SegmentIndex = 0;
 					MeshBatch.MaterialRenderProxy = Section->Material->GetRenderProxy();
 					MeshBatch.ReverseCulling = IsLocalToWorldDeterminantNegative();
 					MeshBatch.Type = PT_TriangleList;
