Propagate CustomPrimitiveData to FDynamicPrimitiveUniformBuffer in FProceduralMeshSceneProxy and FBaseDynamicMeshSceneProxy.

#rnx

diff --git a/Engine/Plugins/Runtime/ProceduralMeshComponent/Source/ProceduralMeshComponent/Private/ProceduralMeshComponent.cpp b/Engine/Plugins/Runtime/ProceduralMeshComponent/Source/ProceduralMeshComponent/Private/ProceduralMeshComponent.cpp
--- a/Engine/Plugins/Runtime/ProceduralMeshComponent/Source/ProceduralMeshComponent/Private/ProceduralMeshComponent.cpp
+++ b/Engine/Plugins/Runtime/ProceduralMeshComponent/Source/ProceduralMeshComponent/Private/ProceduralMeshComponent.cpp
@@ -385,7 +385,7 @@ public:
 						GetScene().GetPrimitiveUniformShaderParameters_RenderThread(GetPrimitiveSceneInfo(), bHasPrecomputedVolumetricLightmap, PreviousLocalToWorld, SingleCaptureIndex, bOutputVelocity);
 
 						FDynamicPrimitiveUniformBuffer& DynamicPrimitiveUniformBuffer = Collector.AllocateOneFrameResource<FDynamicPrimitiveUniformBuffer>();
-						DynamicPrimitiveUniformBuffer.Set(GetLocalToWorld(), PreviousLocalToWorld, GetBounds(), GetLocalBounds(), true, bHasPrecomputedVolumetricLightmap, DrawsVelocity(), bOutputVelocity);
+						DynamicPrimitiveUniformBuffer.Set(GetLocalToWorld(), PreviousLocalToWorld, GetBounds(), GetLocalBounds(), GetLocalBounds(), true, bHasPrecomputedVolumetricLightmap, DrawsVelocity(), bOutputVelocity, GetCustomPrimitiveData());
 						BatchElement.PrimitiveUniformBufferResource = &DynamicPrimitiveUniformBuffer.UniformBuffer;
 
 						BatchElement.FirstIndex = 0;
@@ -502,7 +502,7 @@ public:
 					GetScene().GetPrimitiveUniformShaderParameters_RenderThread(GetPrimitiveSceneInfo(), bHasPrecomputedVolumetricLightmap, PreviousLocalToWorld, SingleCaptureIndex, bOutputVelocity);
 
 					FDynamicPrimitiveUniformBuffer& DynamicPrimitiveUniformBuffer = Context.RayTracingMeshResourceCollector.AllocateOneFrameResource<FDynamicPrimitiveUniformBuffer>();
-					DynamicPrimitiveUniformBuffer.Set(GetLocalToWorld(), PreviousLocalToWorld, GetBounds(), GetLocalBounds(), true, bHasPrecomputedVolumetricLightmap, DrawsVelocity(), bOutputVelocity);
+					DynamicPrimitiveUniformBuffer.Set(GetLocalToWorld(), PreviousLocalToWorld, GetBounds(), GetLocalBounds(), GetLocalBounds(), true, bHasPrecomputedVolumetricLightmap, DrawsVelocity(), bOutputVelocity, GetCustomPrimitiveData());
 					BatchElement.PrimitiveUniformBufferResource = &DynamicPrimitiveUniformBuffer.UniformBuffer;
 
 					BatchElement.FirstIndex = 0;
diff --git a/Engine/Source/Runtime/Engine/Private/SceneManagement.cpp b/Engine/Source/Runtime/Engine/Private/SceneManagement.cpp
--- a/Engine/Source/Runtime/Engine/Private/SceneManagement.cpp
+++ b/Engine/Source/Runtime/Engine/Private/SceneManagement.cpp
@@ -326,7 +326,8 @@ void FDynamicPrimitiveUniformBuffer::Set(
 	bool bReceivesDecals,
 	bool bHasPrecomputedVolumetricLightmap,
 	bool bDrawsVelocity,
-	bool bOutputVelocity)
+	bool bOutputVelocity,
+	const FCustomPrimitiveData* CustomPrimitiveData)
 {
 	check(IsInRenderingThread());
 	UniformBuffer.SetContents(
@@ -342,11 +343,26 @@ void FDynamicPrimitiveUniformBuffer::Set(
 			.OutputVelocity(bOutputVelocity)
 			.DrawsVelocity(bDrawsVelocity)
 			.UseVolumetricLightmap(bHasPrecomputedVolumetricLightmap)
+			.CustomPrimitiveData(CustomPrimitiveData)
 		.Build()
 	);
 	UniformBuffer.InitResource();
 }
 
+void FDynamicPrimitiveUniformBuffer::Set(
+	const FMatrix& LocalToWorld,
+	const FMatrix& PreviousLocalToWorld,
+	const FBoxSphereBounds& WorldBounds,
+	const FBoxSphereBounds& LocalBounds,
+	const FBoxSphereBounds& PreSkinnedLocalBounds,
+	bool bReceivesDecals,
+	bool bHasPrecomputedVolumetricLightmap,
+	bool bDrawsVelocity,
+	bool bOutputVelocity)
+{
+	Set(LocalToWorld, PreviousLocalToWorld, WorldBounds, LocalBounds, PreSkinnedLocalBounds, bReceivesDecals, bHasPrecomputedVolumetricLightmap, bDrawsVelocity, bOutputVelocity, nullptr);
+}
+
 void FDynamicPrimitiveUniformBuffer::Set(
 	const FMatrix& LocalToWorld,
 	const FMatrix& PreviousLocalToWorld,
@@ -357,7 +373,7 @@ void FDynamicPrimitiveUniformBuffer::Set(
 	bool bDrawsVelocity,
 	bool bOutputVelocity)
 {
-	Set(LocalToWorld, PreviousLocalToWorld, WorldBounds, LocalBounds, LocalBounds, bReceivesDecals, bHasPrecomputedVolumetricLightmap, bDrawsVelocity, bOutputVelocity);
+	Set(LocalToWorld, PreviousLocalToWorld, WorldBounds, LocalBounds, LocalBounds, bReceivesDecals, bHasPrecomputedVolumetricLightmap, bDrawsVelocity, bOutputVelocity, nullptr);
 }
 
 FLightMapInteraction FLightMapInteraction::Texture(
diff --git a/Engine/Source/Runtime/Engine/Public/SceneManagement.h b/Engine/Source/Runtime/Engine/Public/SceneManagement.h
--- a/Engine/Source/Runtime/Engine/Public/SceneManagement.h
+++ b/Engine/Source/Runtime/Engine/Public/SceneManagement.h
@@ -2579,6 +2579,18 @@ public:
 
 	TUniformBuffer<FPrimitiveUniformShaderParameters> UniformBuffer;
 
+	ENGINE_API void Set(
+		const FMatrix& LocalToWorld,
+		const FMatrix& PreviousLocalToWorld,
+		const FBoxSphereBounds& WorldBounds,
+		const FBoxSphereBounds& LocalBounds,
+		const FBoxSphereBounds& PreSkinnedLocalBounds,
+		bool bReceivesDecals,
+		bool bHasPrecomputedVolumetricLightmap,
+		bool bDrawsVelocity,
+		bool bOutputVelocity,
+		const FCustomPrimitiveData* CustomPrimitiveData);
+
 	ENGINE_API void Set(
 		const FMatrix& LocalToWorld,
 		const FMatrix& PreviousLocalToWorld,
diff --git a/Engine/Source/Runtime/GeometryFramework/Public/Components/BaseDynamicMeshSceneProxy.h b/Engine/Source/Runtime/GeometryFramework/Public/Components/BaseDynamicMeshSceneProxy.h
--- a/Engine/Source/Runtime/GeometryFramework/Public/Components/BaseDynamicMeshSceneProxy.h
+++ b/Engine/Source/Runtime/GeometryFramework/Public/Components/BaseDynamicMeshSceneProxy.h
@@ -1007,7 +1007,7 @@ public:
 
 					// do we need separate one of these for each MeshRenderBufferSet?
 					FDynamicPrimitiveUniformBuffer& DynamicPrimitiveUniformBuffer = Collector.AllocateOneFrameResource<FDynamicPrimitiveUniformBuffer>();
-					DynamicPrimitiveUniformBuffer.Set(GetLocalToWorld(), PreviousLocalToWorld, GetBounds(), GetLocalBounds(), true, bHasPrecomputedVolumetricLightmap, DrawsVelocity(), bOutputVelocity);
+					DynamicPrimitiveUniformBuffer.Set(GetLocalToWorld(), PreviousLocalToWorld, GetBounds(), GetLocalBounds(), GetLocalBounds(), true, bHasPrecomputedVolumetricLightmap, DrawsVelocity(), bOutputVelocity, GetCustomPrimitiveData());
 
 					if (BufferSet->IndexBuffer.Indices.Num() > 0)
 					{
