Fixed crash in procedural mesh component with ray tracing enabled

#github pr 8806
#8806

diff --git a/Engine/Plugins/Runtime/ProceduralMeshComponent/Source/ProceduralMeshComponent/Private/ProceduralMeshComponent.cpp b/Engine/Plugins/Runtime/ProceduralMeshComponent/Source/ProceduralMeshComponent/Private/ProceduralMeshComponent.cpp
--- a/Engine/Plugins/Runtime/ProceduralMeshComponent/Source/ProceduralMeshComponent/Private/ProceduralMeshComponent.cpp
+++ b/Engine/Plugins/Runtime/ProceduralMeshComponent/Source/ProceduralMeshComponent/Private/ProceduralMeshComponent.cpp
@@ -315,6 +315,7 @@ public:
 					FRayTracingGeometrySegment Segment;
 					Segment.VertexBuffer = Section->VertexBuffers.PositionVertexBuffer.VertexBufferRHI;
 					Segment.NumPrimitives = Section->RayTracingGeometry.Initializer.TotalPrimitiveCount;
+					Segment.MaxVertices = Section->VertexBuffers.PositionVertexBuffer.GetNumVertices();
 					Section->RayTracingGeometry.Initializer.Segments.Add(Segment);
 
 					Section->RayTracingGeometry.UpdateRHI();
