Add support for bodysetup trimesh estimates for all implementation of IInterface_CollisionDataProvider

#rnx

diff --git a/Engine/Plugins/Enterprise/LidarPointCloud/Source/LidarPointCloudRuntime/Private/LidarPointCloud.cpp b/Engine/Plugins/Enterprise/LidarPointCloud/Source/LidarPointCloudRuntime/Private/LidarPointCloud.cpp
--- a/Engine/Plugins/Enterprise/LidarPointCloud/Source/LidarPointCloudRuntime/Private/LidarPointCloud.cpp
+++ b/Engine/Plugins/Enterprise/LidarPointCloud/Source/LidarPointCloudRuntime/Private/LidarPointCloud.cpp
@@ -1210,6 +1210,13 @@ void ULidarPointCloud::CalculateNormals(TArray64<FLidarPointCloudPoint*>* Points
 	});
 }
 
+bool ULidarPointCloud::GetTriMeshSizeEstimates(struct FTriMeshCollisionDataEstimates& OutTriMeshEstimates, bool bInUseAllTriData) const
+{
+	OutTriMeshEstimates.VerticeCount = Octree.GetCollisionData()->Vertices.Num();
+	
+	return true;
+}
+
 bool ULidarPointCloud::GetPhysicsTriMeshData(FTriMeshCollisionData* CollisionData, bool InUseAllTriData)
 {
 	CollisionData->Vertices = Octree.GetCollisionData()->Vertices;
diff --git a/Engine/Plugins/Enterprise/LidarPointCloud/Source/LidarPointCloudRuntime/Public/LidarPointCloud.h b/Engine/Plugins/Enterprise/LidarPointCloud/Source/LidarPointCloudRuntime/Public/LidarPointCloud.h
--- a/Engine/Plugins/Enterprise/LidarPointCloud/Source/LidarPointCloudRuntime/Public/LidarPointCloud.h
+++ b/Engine/Plugins/Enterprise/LidarPointCloud/Source/LidarPointCloudRuntime/Public/LidarPointCloud.h
@@ -592,6 +592,7 @@ public:
 
 	//~ Begin Interface_CollisionDataProvider Interface
 	virtual bool GetPhysicsTriMeshData(FTriMeshCollisionData* CollisionData, bool InUseAllTriData) override;
+	virtual bool GetTriMeshSizeEstimates(struct FTriMeshCollisionDataEstimates& OutTriMeshEstimates, bool bInUseAllTriData) const override;
 	virtual bool ContainsPhysicsTriMeshData(bool InUseAllTriData) const override { return HasCollisionData(); }
 	virtual bool WantsNegXTriMesh() override { return false; }
 	//~ End Interface_CollisionDataProvider Interface
diff --git a/Engine/Plugins/Runtime/ProceduralMeshComponent/Source/ProceduralMeshComponent/Private/ProceduralMeshComponent.cpp b/Engine/Plugins/Runtime/ProceduralMeshComponent/Source/ProceduralMeshComponent/Private/ProceduralMeshComponent.cpp
--- a/Engine/Plugins/Runtime/ProceduralMeshComponent/Source/ProceduralMeshComponent/Private/ProceduralMeshComponent.cpp
+++ b/Engine/Plugins/Runtime/ProceduralMeshComponent/Source/ProceduralMeshComponent/Private/ProceduralMeshComponent.cpp
@@ -949,6 +949,19 @@ FBoxSphereBounds UProceduralMeshComponent::CalcBounds(const FTransform& LocalToW
 	return Ret;
 }
 
+bool UProceduralMeshComponent::GetTriMeshSizeEstimates(struct FTriMeshCollisionDataEstimates& OutTriMeshEstimates, bool bInUseAllTriData) const
+{
+	for (const FProcMeshSection& Section : ProcMeshSections)
+	{
+		if (Section.bEnableCollision)
+		{
+			OutTriMeshEstimates.VerticeCount += Section.ProcVertexBuffer.Num();
+		}
+	}
+
+	return true;
+}
+
 bool UProceduralMeshComponent::GetPhysicsTriMeshData(struct FTriMeshCollisionData* CollisionData, bool InUseAllTriData)
 {
 	int32 VertexBase = 0; // Base vertex index for current section
diff --git a/Engine/Plugins/Runtime/ProceduralMeshComponent/Source/ProceduralMeshComponent/Public/ProceduralMeshComponent.h b/Engine/Plugins/Runtime/ProceduralMeshComponent/Source/ProceduralMeshComponent/Public/ProceduralMeshComponent.h
--- a/Engine/Plugins/Runtime/ProceduralMeshComponent/Source/ProceduralMeshComponent/Public/ProceduralMeshComponent.h
+++ b/Engine/Plugins/Runtime/ProceduralMeshComponent/Source/ProceduralMeshComponent/Public/ProceduralMeshComponent.h
@@ -260,6 +260,7 @@ public:
 	void SetCollisionConvexMeshes(const TArray< TArray<FVector> >& ConvexMeshes);
 
 	//~ Begin Interface_CollisionDataProvider Interface
+	virtual bool GetTriMeshSizeEstimates(struct FTriMeshCollisionDataEstimates& OutTriMeshEstimates, bool bInUseAllTriData) const override;
 	virtual bool GetPhysicsTriMeshData(struct FTriMeshCollisionData* CollisionData, bool InUseAllTriData) override;
 	virtual bool ContainsPhysicsTriMeshData(bool InUseAllTriData) const override;
 	virtual bool WantsNegXTriMesh() override{ return false; }
diff --git a/Engine/Source/Runtime/Engine/Classes/Components/ModelComponent.h b/Engine/Source/Runtime/Engine/Classes/Components/ModelComponent.h
--- a/Engine/Source/Runtime/Engine/Classes/Components/ModelComponent.h
+++ b/Engine/Source/Runtime/Engine/Classes/Components/ModelComponent.h
@@ -125,6 +125,7 @@ public:
 
 	//~ Begin Interface_CollisionDataProvider Interface
 	virtual bool GetPhysicsTriMeshData(struct FTriMeshCollisionData* CollisionData, bool InUseAllTriData) override;
+	virtual bool GetTriMeshSizeEstimates(struct FTriMeshCollisionDataEstimates& OutTriMeshEstimates, bool bInUseAllTriData) const override;
 	virtual bool ContainsPhysicsTriMeshData(bool InUseAllTriData) const override;
 	virtual bool WantsNegXTriMesh() override { return false; }
 	//~ End Interface_CollisionDataProvider Interface
diff --git a/Engine/Source/Runtime/Engine/Classes/Engine/SkeletalMesh.h b/Engine/Source/Runtime/Engine/Classes/Engine/SkeletalMesh.h
--- a/Engine/Source/Runtime/Engine/Classes/Engine/SkeletalMesh.h
+++ b/Engine/Source/Runtime/Engine/Classes/Engine/SkeletalMesh.h
@@ -2962,6 +2962,7 @@ public:
 
 	//~ Begin Interface_CollisionDataProvider Interface
 	virtual bool GetPhysicsTriMeshData(struct FTriMeshCollisionData* CollisionData, bool InUseAllTriData) override;
+	virtual bool GetTriMeshSizeEstimates(struct FTriMeshCollisionDataEstimates& OutTriMeshEstimates, bool bInUseAllTriData) const override;
 	virtual bool ContainsPhysicsTriMeshData(bool InUseAllTriData) const override;
 	virtual bool WantsNegXTriMesh() override
 	{
diff --git a/Engine/Source/Runtime/Engine/Private/Components/ModelComponent.cpp b/Engine/Source/Runtime/Engine/Private/Components/ModelComponent.cpp
--- a/Engine/Source/Runtime/Engine/Private/Components/ModelComponent.cpp
+++ b/Engine/Source/Runtime/Engine/Private/Components/ModelComponent.cpp
@@ -631,6 +631,13 @@ void UModelComponent::InvalidateCollisionData()
 }
 #endif // WITH_EDITOR
 
+bool UModelComponent::GetTriMeshSizeEstimates(struct FTriMeshCollisionDataEstimates& OutTriMeshEstimates, bool bInUseAllTriData) const
+{
+	check(Model);
+	OutTriMeshEstimates.VerticeCount = Model->VertexBuffer.Vertices.Num();
+	return true;
+}
+
 bool UModelComponent::GetPhysicsTriMeshData(struct FTriMeshCollisionData* CollisionData, bool InUseAllTriData)
 {
 	check(Model);
diff --git a/Engine/Source/Runtime/Engine/Private/SkeletalMesh.cpp b/Engine/Source/Runtime/Engine/Private/SkeletalMesh.cpp
--- a/Engine/Source/Runtime/Engine/Private/SkeletalMesh.cpp
+++ b/Engine/Source/Runtime/Engine/Private/SkeletalMesh.cpp
@@ -5057,6 +5057,18 @@ bool USkeletalMesh::ContainsPhysicsTriMeshData(bool InUseAllTriData) const
 	return GetEnablePerPolyCollision();
 }
 
+bool USkeletalMesh::GetTriMeshSizeEstimates(struct FTriMeshCollisionDataEstimates& OutTriMeshEstimates, bool bInUseAllTriData) const
+{
+#if WITH_EDITORONLY_DATA
+	if (GetResourceForRendering() && GetEnablePerPolyCollision())
+	{
+		OutTriMeshEstimates.VerticeCount = GetResourceForRendering()->LODRenderData[0].GetNumVertices();
+	}
+#endif // #if WITH_EDITORONLY_DATA
+
+	return true;
+}
+
 bool USkeletalMesh::GetPhysicsTriMeshData(FTriMeshCollisionData* CollisionData, bool bInUseAllTriData)
 {
 #if WITH_EDITORONLY_DATA
diff --git a/Engine/Source/Runtime/MRMesh/Private/MRMeshComponent.cpp b/Engine/Source/Runtime/MRMesh/Private/MRMeshComponent.cpp
--- a/Engine/Source/Runtime/MRMesh/Private/MRMeshComponent.cpp
+++ b/Engine/Source/Runtime/MRMesh/Private/MRMeshComponent.cpp
@@ -1228,6 +1228,16 @@ void UMRMeshBodyHolder::Cleanup()
 	BodyInstance.TermBody();
 }
 
+bool UMRMeshBodyHolder::GetTriMeshSizeEstimates(struct FTriMeshCollisionDataEstimates& OutTriMeshEstimates, bool bInUseAllTriData) const
+{
+	if (BrickDataReceipt.IsValid() && Indices->Num() > 0)
+	{
+		OutTriMeshEstimates.VerticeCount = PositionData->Num();
+	}
+
+	return true;
+}
+
 bool UMRMeshBodyHolder::GetPhysicsTriMeshData(struct FTriMeshCollisionData* CollisionData, bool InUseAllTriData)
 {
 	if (BrickDataReceipt.IsValid() && Indices->Num() > 0)
diff --git a/Engine/Source/Runtime/MRMesh/Public/MRMeshComponent.h b/Engine/Source/Runtime/MRMesh/Public/MRMeshComponent.h
--- a/Engine/Source/Runtime/MRMesh/Public/MRMeshComponent.h
+++ b/Engine/Source/Runtime/MRMesh/Public/MRMeshComponent.h
@@ -79,6 +79,7 @@ public:
 
 	//~ Begin Interface_CollisionDataProvider Interface
 	virtual bool GetPhysicsTriMeshData(struct FTriMeshCollisionData* CollisionData, bool InUseAllTriData) override;
+	virtual bool GetTriMeshSizeEstimates(struct FTriMeshCollisionDataEstimates& OutTriMeshEstimates, bool bInUseAllTriData) const override;
 	virtual bool ContainsPhysicsTriMeshData(bool InUseAllTriData) const override;
 	virtual bool WantsNegXTriMesh() override { return false; }
 	//~ End Interface_CollisionDataProvider Interface
