Fix formatting of string literals that could crash the engine.

diff --git a/Engine/Plugins/Runtime/ProceduralMeshComponent/Source/ProceduralMeshComponent/Private/ProceduralMeshComponent.cpp b/Engine/Plugins/Runtime/ProceduralMeshComponent/Source/ProceduralMeshComponent/Private/ProceduralMeshComponent.cpp
--- a/Engine/Plugins/Runtime/ProceduralMeshComponent/Source/ProceduralMeshComponent/Private/ProceduralMeshComponent.cpp
+++ b/Engine/Plugins/Runtime/ProceduralMeshComponent/Source/ProceduralMeshComponent/Private/ProceduralMeshComponent.cpp
@@ -626,7 +626,7 @@ void UProceduralMeshComponent::CreateMeshSection(int32 SectionIndex, const TArra
 	if (NumDegenerateTriangles > 0)
 	{
 		UE_LOG(LogProceduralComponent, Warning, TEXT("Detected %d degenerate triangle%s with non-unique vertex indices for created mesh section in '%s'; degenerate triangles will be dropped."),
-			   NumDegenerateTriangles, NumDegenerateTriangles > 1 ? "s" : "", *GetFullName());
+			   NumDegenerateTriangles, NumDegenerateTriangles > 1 ? TEXT("s") : TEXT(""), *GetFullName());
 	}
 
 	// Copy index buffer for non-degenerate triangles
