Allow caller to choose whether to perform sRGB conversion for ProceduralMeshComponent's CreateMeshSection and UpdateMeshSection functions.
To keep current results unchanged, set the default behavior for the two functions to be different (Create doesn't convert, Update does)

diff --git a/Engine/Plugins/Runtime/ProceduralMeshComponent/Source/ProceduralMeshComponent/Private/ProceduralMeshComponent.cpp b/Engine/Plugins/Runtime/ProceduralMeshComponent/Source/ProceduralMeshComponent/Private/ProceduralMeshComponent.cpp
--- a/Engine/Plugins/Runtime/ProceduralMeshComponent/Source/ProceduralMeshComponent/Private/ProceduralMeshComponent.cpp
+++ b/Engine/Plugins/Runtime/ProceduralMeshComponent/Source/ProceduralMeshComponent/Private/ProceduralMeshComponent.cpp
@@ -554,7 +554,7 @@ void UProceduralMeshComponent::PostLoad()
 	}
 }
 
-void UProceduralMeshComponent::CreateMeshSection_LinearColor(int32 SectionIndex, const TArray<FVector>& Vertices, const TArray<int32>& Triangles, const TArray<FVector>& Normals, const TArray<FVector2D>& UV0, const TArray<FVector2D>& UV1, const TArray<FVector2D>& UV2, const TArray<FVector2D>& UV3, const TArray<FLinearColor>& VertexColors, const TArray<FProcMeshTangent>& Tangents, bool bCreateCollision)
+void UProceduralMeshComponent::CreateMeshSection_LinearColor(int32 SectionIndex, const TArray<FVector>& Vertices, const TArray<int32>& Triangles, const TArray<FVector>& Normals, const TArray<FVector2D>& UV0, const TArray<FVector2D>& UV1, const TArray<FVector2D>& UV2, const TArray<FVector2D>& UV3, const TArray<FLinearColor>& VertexColors, const TArray<FProcMeshTangent>& Tangents, bool bCreateCollision, bool bSRGBConversion)
 {
 	// Convert FLinearColors to FColors
 	TArray<FColor> Colors;
@@ -564,7 +564,7 @@ void UProceduralMeshComponent::CreateMeshSection_LinearColor(int32 SectionIndex,
 
 		for (int32 ColorIdx = 0; ColorIdx < VertexColors.Num(); ColorIdx++)
 		{
-			Colors[ColorIdx] = VertexColors[ColorIdx].ToFColor(false);
+			Colors[ColorIdx] = VertexColors[ColorIdx].ToFColor(bSRGBConversion);
 		}
 	}
 
@@ -661,7 +661,7 @@ void UProceduralMeshComponent::CreateMeshSection(int32 SectionIndex, const TArra
 	MarkRenderStateDirty(); // New section requires recreating scene proxy
 }
 
-void UProceduralMeshComponent::UpdateMeshSection_LinearColor(int32 SectionIndex, const TArray<FVector>& Vertices, const TArray<FVector>& Normals, const TArray<FVector2D>& UV0, const TArray<FVector2D>& UV1, const TArray<FVector2D>& UV2, const TArray<FVector2D>& UV3, const TArray<FLinearColor>& VertexColors, const TArray<FProcMeshTangent>& Tangents)
+void UProceduralMeshComponent::UpdateMeshSection_LinearColor(int32 SectionIndex, const TArray<FVector>& Vertices, const TArray<FVector>& Normals, const TArray<FVector2D>& UV0, const TArray<FVector2D>& UV1, const TArray<FVector2D>& UV2, const TArray<FVector2D>& UV3, const TArray<FLinearColor>& VertexColors, const TArray<FProcMeshTangent>& Tangents, bool bSRGBConversion)
 {
 	// Convert FLinearColors to FColors
 	TArray<FColor> Colors;
@@ -671,7 +671,7 @@ void UProceduralMeshComponent::UpdateMeshSection_LinearColor(int32 SectionIndex,
 
 		for (int32 ColorIdx = 0; ColorIdx < VertexColors.Num(); ColorIdx++)
 		{
-			Colors[ColorIdx] = VertexColors[ColorIdx].ToFColor(true);
+			Colors[ColorIdx] = VertexColors[ColorIdx].ToFColor(bSRGBConversion);
 		}
 	}
 
diff --git a/Engine/Plugins/Runtime/ProceduralMeshComponent/Source/ProceduralMeshComponent/Public/ProceduralMeshComponent.h b/Engine/Plugins/Runtime/ProceduralMeshComponent/Source/ProceduralMeshComponent/Public/ProceduralMeshComponent.h
--- a/Engine/Plugins/Runtime/ProceduralMeshComponent/Source/ProceduralMeshComponent/Public/ProceduralMeshComponent.h
+++ b/Engine/Plugins/Runtime/ProceduralMeshComponent/Source/ProceduralMeshComponent/Public/ProceduralMeshComponent.h
@@ -183,14 +183,16 @@ public:
 	 *	@param	VertexColors		Optional array of colors for each vertex. If supplied, must be same length as Vertices array.
 	 *	@param	Tangents			Optional array of tangent vector for each vertex. If supplied, must be same length as Vertices array.
 	 *	@param	bCreateCollision	Indicates whether collision should be created for this section. This adds significant cost.
+	 *	@param	bSRGBConversion		Whether to do sRGB conversion when converting FLinearColor to FColor
 	 */
 	UFUNCTION(BlueprintCallable, Category = "Components|ProceduralMesh", meta = (DisplayName = "Create Mesh Section", AutoCreateRefTerm = "Normals,UV0,UV1,UV2,UV3,VertexColors,Tangents", AdvancedDisplay = "UV1,UV2,UV3"))
-	void CreateMeshSection_LinearColor(int32 SectionIndex, const TArray<FVector>& Vertices, const TArray<int32>& Triangles, const TArray<FVector>& Normals, const TArray<FVector2D>& UV0, const TArray<FVector2D>& UV1, const TArray<FVector2D>& UV2, const TArray<FVector2D>& UV3, const TArray<FLinearColor>& VertexColors, const TArray<FProcMeshTangent>& Tangents, bool bCreateCollision);
+	void CreateMeshSection_LinearColor(int32 SectionIndex, const TArray<FVector>& Vertices, const TArray<int32>& Triangles, const TArray<FVector>& Normals, const TArray<FVector2D>& UV0, const TArray<FVector2D>& UV1, const TArray<FVector2D>& UV2, const TArray<FVector2D>& UV3,
+		const TArray<FLinearColor>& VertexColors, const TArray<FProcMeshTangent>& Tangents, bool bCreateCollision, UPARAM(DisplayName = "SRGB Conversion") bool bSRGBConversion = false);
 
-	void CreateMeshSection_LinearColor(int32 SectionIndex, const TArray<FVector>& Vertices, const TArray<int32>& Triangles, const TArray<FVector>& Normals, const TArray<FVector2D>& UV0, const TArray<FLinearColor>& VertexColors, const TArray<FProcMeshTangent>& Tangents, bool bCreateCollision)
+	void CreateMeshSection_LinearColor(int32 SectionIndex, const TArray<FVector>& Vertices, const TArray<int32>& Triangles, const TArray<FVector>& Normals, const TArray<FVector2D>& UV0, const TArray<FLinearColor>& VertexColors, const TArray<FProcMeshTangent>& Tangents, bool bCreateCollision, bool bSRGBConversion = false)
 	{
 		TArray<FVector2D> EmptyArray;
-		CreateMeshSection_LinearColor(SectionIndex, Vertices, Triangles, Normals, UV0, EmptyArray, EmptyArray, EmptyArray, VertexColors, Tangents, bCreateCollision);
+		CreateMeshSection_LinearColor(SectionIndex, Vertices, Triangles, Normals, UV0, EmptyArray, EmptyArray, EmptyArray, VertexColors, Tangents, bCreateCollision, bSRGBConversion);
 	}
 
 	/**
@@ -218,14 +220,16 @@ public:
 	 *	@param	UV0					Optional array of texture co-ordinates for each vertex. If supplied, must be same length as Vertices array.
 	 *	@param	VertexColors		Optional array of colors for each vertex. If supplied, must be same length as Vertices array.
 	 *	@param	Tangents			Optional array of tangent vector for each vertex. If supplied, must be same length as Vertices array.
+	 *	@param	bSRGBConversion		Whether to do sRGB conversion when converting FLinearColor to FColor
 	 */
 	UFUNCTION(BlueprintCallable, Category = "Components|ProceduralMesh", meta = (DisplayName = "Update Mesh Section", AutoCreateRefTerm = "Normals,UV0,UV1,UV2,UV3,VertexColors,Tangents", AdvancedDisplay = "UV1,UV2,UV3"))
-	void UpdateMeshSection_LinearColor(int32 SectionIndex, const TArray<FVector>& Vertices, const TArray<FVector>& Normals, const TArray<FVector2D>& UV0, const TArray<FVector2D>& UV1, const TArray<FVector2D>& UV2, const TArray<FVector2D>& UV3, const TArray<FLinearColor>& VertexColors, const TArray<FProcMeshTangent>& Tangents);
+	void UpdateMeshSection_LinearColor(int32 SectionIndex, const TArray<FVector>& Vertices, const TArray<FVector>& Normals, const TArray<FVector2D>& UV0, const TArray<FVector2D>& UV1, const TArray<FVector2D>& UV2, const TArray<FVector2D>& UV3,
+		const TArray<FLinearColor>& VertexColors, const TArray<FProcMeshTangent>& Tangents, UPARAM(DisplayName = "SRGB Conversion") bool bSRGBConversion = true);
 
-	void UpdateMeshSection_LinearColor(int32 SectionIndex, const TArray<FVector>& Vertices, const TArray<FVector>& Normals, const TArray<FVector2D>& UV0, const TArray<FLinearColor>& VertexColors, const TArray<FProcMeshTangent>& Tangents)
+	void UpdateMeshSection_LinearColor(int32 SectionIndex, const TArray<FVector>& Vertices, const TArray<FVector>& Normals, const TArray<FVector2D>& UV0, const TArray<FLinearColor>& VertexColors, const TArray<FProcMeshTangent>& Tangents, bool bSRGBConversion = true)
 	{
 		TArray<FVector2D> EmptyArray;
-		UpdateMeshSection_LinearColor(SectionIndex, Vertices, Normals, UV0, EmptyArray, EmptyArray, EmptyArray, VertexColors, Tangents);
+		UpdateMeshSection_LinearColor(SectionIndex, Vertices, Normals, UV0, EmptyArray, EmptyArray, EmptyArray, VertexColors, Tangents, bSRGBConversion);
 	}
 
 	/** Clear a section of the procedural mesh. Other sections do not change index. */
