Make ProceduralComponent honor the bReceivesDecals setting.

diff --git a/Engine/Plugins/Runtime/ProceduralMeshComponent/Source/ProceduralMeshComponent/Private/ProceduralMeshComponent.cpp b/Engine/Plugins/Runtime/ProceduralMeshComponent/Source/ProceduralMeshComponent/Private/ProceduralMeshComponent.cpp
--- a/Engine/Plugins/Runtime/ProceduralMeshComponent/Source/ProceduralMeshComponent/Private/ProceduralMeshComponent.cpp
+++ b/Engine/Plugins/Runtime/ProceduralMeshComponent/Source/ProceduralMeshComponent/Private/ProceduralMeshComponent.cpp
@@ -389,7 +389,7 @@ public:
 						bOutputVelocity |= AlwaysHasVelocity();
 
 						FDynamicPrimitiveUniformBuffer& DynamicPrimitiveUniformBuffer = Collector.AllocateOneFrameResource<FDynamicPrimitiveUniformBuffer>();
-						DynamicPrimitiveUniformBuffer.Set(GetLocalToWorld(), PreviousLocalToWorld, GetBounds(), GetLocalBounds(), GetLocalBounds(), true, bHasPrecomputedVolumetricLightmap, bOutputVelocity, GetCustomPrimitiveData());
+						DynamicPrimitiveUniformBuffer.Set(GetLocalToWorld(), PreviousLocalToWorld, GetBounds(), GetLocalBounds(), GetLocalBounds(), ReceivesDecals(), bHasPrecomputedVolumetricLightmap, bOutputVelocity, GetCustomPrimitiveData());
 						BatchElement.PrimitiveUniformBufferResource = &DynamicPrimitiveUniformBuffer.UniformBuffer;
 
 						BatchElement.FirstIndex = 0;
@@ -507,7 +507,7 @@ public:
 					bOutputVelocity |= AlwaysHasVelocity();
 
 					FDynamicPrimitiveUniformBuffer& DynamicPrimitiveUniformBuffer = Context.RayTracingMeshResourceCollector.AllocateOneFrameResource<FDynamicPrimitiveUniformBuffer>();
-					DynamicPrimitiveUniformBuffer.Set(GetLocalToWorld(), PreviousLocalToWorld, GetBounds(), GetLocalBounds(), GetLocalBounds(), true, bHasPrecomputedVolumetricLightmap, bOutputVelocity, GetCustomPrimitiveData());
+					DynamicPrimitiveUniformBuffer.Set(GetLocalToWorld(), PreviousLocalToWorld, GetBounds(), GetLocalBounds(), GetLocalBounds(), ReceivesDecals(), bHasPrecomputedVolumetricLightmap, bOutputVelocity, GetCustomPrimitiveData());
 					BatchElement.PrimitiveUniformBufferResource = &DynamicPrimitiveUniformBuffer.UniformBuffer;
 
 					BatchElement.FirstIndex = 0;
