Fix issue where Static Mesh's created via animation tool defaulted to Game/Meshes when that may not be a valid folder

diff --git a/Engine/Plugins/Runtime/ProceduralMeshComponent/Source/ProceduralMeshComponentEditor/Private/ProceduralMeshComponentDetails.cpp b/Engine/Plugins/Runtime/ProceduralMeshComponent/Source/ProceduralMeshComponentEditor/Private/ProceduralMeshComponentDetails.cpp
--- a/Engine/Plugins/Runtime/ProceduralMeshComponent/Source/ProceduralMeshComponentEditor/Private/ProceduralMeshComponentDetails.cpp
+++ b/Engine/Plugins/Runtime/ProceduralMeshComponent/Source/ProceduralMeshComponentEditor/Private/ProceduralMeshComponentDetails.cpp
@@ -15,6 +15,7 @@
 #include "DetailLayoutBuilder.h"
 #include "DetailWidgetRow.h"
 #include "DetailCategoryBuilder.h"
+#include "EditorDirectories.h"
 #include "IDetailsView.h"
 #include "ProceduralMeshComponent.h"
 #include "StaticMeshAttributes.h"
@@ -90,7 +91,16 @@ FReply FProceduralMeshComponentDetails::ClickedOnConvertToStaticMesh()
 	if (ProcMeshComp != nullptr)
 	{
 		FString NewNameSuggestion = FString(TEXT("ProcMesh"));
-		FString PackageName = FString(TEXT("/Game/Meshes/")) + NewNameSuggestion;
+		FString DefaultPath;
+		const FString DefaultDirectory = FEditorDirectories::Get().GetLastDirectory(ELastDirectory::NEW_ASSET);
+		FPackageName::TryConvertFilenameToLongPackageName(DefaultDirectory, DefaultPath);
+
+		if (DefaultPath.IsEmpty())
+		{
+			DefaultPath = TEXT("/Game/Meshes");
+		}
+
+		FString PackageName = DefaultPath / NewNameSuggestion;
 		FString Name;
 		FAssetToolsModule& AssetToolsModule = FModuleManager::LoadModuleChecked<FAssetToolsModule>("AssetTools");
 		AssetToolsModule.Get().CreateUniqueAssetName(PackageName, TEXT(""), PackageName, Name);
diff --git a/Engine/Source/Developer/MeshUtilities/Private/MeshUtilities.cpp b/Engine/Source/Developer/MeshUtilities/Private/MeshUtilities.cpp
--- a/Engine/Source/Developer/MeshUtilities/Private/MeshUtilities.cpp
+++ b/Engine/Source/Developer/MeshUtilities/Private/MeshUtilities.cpp
@@ -25,6 +25,7 @@
 #include "Components/ShapeComponent.h"
 #include "Engine/SkinnedAssetCommon.h"
 #include "Engine/StaticMesh.h"
+#include "EditorDirectories.h"
 #include "Materials/Material.h"
 #include "RawMesh.h"
 #include "StaticMeshResources.h"
@@ -500,7 +501,16 @@ UStaticMesh* FMeshUtilities::ConvertMeshesToStaticMesh(const TArray<UMeshCompone
 	if (InPackageName.IsEmpty())
 	{
 		FString NewNameSuggestion = FString(TEXT("StaticMesh"));
-		FString PackageNameSuggestion = FString(TEXT("/Game/Meshes/")) + NewNameSuggestion;
+		FString DefaultPath;
+		const FString DefaultDirectory = FEditorDirectories::Get().GetLastDirectory(ELastDirectory::NEW_ASSET);
+		FPackageName::TryConvertFilenameToLongPackageName(DefaultDirectory, DefaultPath);
+
+		if (DefaultPath.IsEmpty())
+		{
+			DefaultPath = TEXT("/Game/Meshes");
+		}
+
+		FString PackageNameSuggestion = DefaultPath / NewNameSuggestion;
 		FString Name;
 		FAssetToolsModule& AssetToolsModule = FModuleManager::LoadModuleChecked<FAssetToolsModule>("AssetTools");
 		AssetToolsModule.Get().CreateUniqueAssetName(PackageNameSuggestion, TEXT(""), PackageNameSuggestion, Name);
