Refactor RayTracingGeometry initalization in a few proxies to not modify FRayTracingGeometry::Initializer directly.
- some foundation work towards making FRayTracingGeometry::Initializer private.

diff --git a/Engine/Plugins/Experimental/ImagePlate/Source/ImagePlate/Private/ImagePlateComponent.cpp b/Engine/Plugins/Experimental/ImagePlate/Source/ImagePlate/Private/ImagePlateComponent.cpp
--- a/Engine/Plugins/Experimental/ImagePlate/Source/ImagePlate/Private/ImagePlateComponent.cpp
+++ b/Engine/Plugins/Experimental/ImagePlate/Source/ImagePlate/Private/ImagePlateComponent.cpp
@@ -115,25 +115,21 @@ namespace
 			{
 				FRayTracingGeometryInitializer Initializer;
 				Initializer.DebugName = GetOwnerName();
-				Initializer.IndexBuffer = nullptr;
-				Initializer.TotalPrimitiveCount = 0;
+				Initializer.IndexBuffer = IndexBuffer.IndexBufferRHI;
+				Initializer.TotalPrimitiveCount = IndexBuffer.Indices.Num() / 3;
 				Initializer.GeometryType = RTGT_Triangles;
 				Initializer.bFastBuild = true;
 				Initializer.bAllowUpdate = false;
 
-				RayTracingGeometry.SetInitializer(Initializer);
-				RayTracingGeometry.InitResource(RHICmdList);
-
-				RayTracingGeometry.Initializer.IndexBuffer = IndexBuffer.IndexBufferRHI;
-				RayTracingGeometry.Initializer.TotalPrimitiveCount = IndexBuffer.Indices.Num() / 3;
-
 				FRayTracingGeometrySegment Segment;
 				Segment.VertexBuffer = VertexBuffers.PositionVertexBuffer.VertexBufferRHI;
-				Segment.NumPrimitives = RayTracingGeometry.Initializer.TotalPrimitiveCount;
 				Segment.MaxVertices = VertexBuffers.PositionVertexBuffer.GetNumVertices();
-				RayTracingGeometry.Initializer.Segments.Add(Segment);
+				Segment.NumPrimitives = Initializer.TotalPrimitiveCount;
 
-				RayTracingGeometry.UpdateRHI(RHICmdList);
+				Initializer.Segments.Add(Segment);
+
+				RayTracingGeometry.SetInitializer(Initializer);
+				RayTracingGeometry.InitResource(RHICmdList);
 			}
 #endif
 		}
diff --git a/Engine/Plugins/Runtime/CableComponent/Source/CableComponent/Private/CableComponent.cpp b/Engine/Plugins/Runtime/CableComponent/Source/CableComponent/Private/CableComponent.cpp
--- a/Engine/Plugins/Runtime/CableComponent/Source/CableComponent/Private/CableComponent.cpp
+++ b/Engine/Plugins/Runtime/CableComponent/Source/CableComponent/Private/CableComponent.cpp
@@ -583,8 +583,7 @@ private:
 #if RHI_RAYTRACING
 	void CreateDynamicRayTracingGeometries(FRHICommandListBase& RHICmdList)
 	{
-		FRayTracingGeometryInitializer& Initializer = DynamicRayTracingGeometry.Initializer;
-		Initializer = StaticRayTracingGeometry.Initializer;
+		FRayTracingGeometryInitializer Initializer = StaticRayTracingGeometry.Initializer;
 		for (FRayTracingGeometrySegment& Segment : Initializer.Segments)
 		{
 			Segment.VertexBuffer = nullptr;
@@ -593,6 +592,7 @@ private:
 		Initializer.bFastBuild = true;
 		Initializer.Type = ERayTracingGeometryInitializerType::Rendering;
 
+		DynamicRayTracingGeometry.SetInitializer(MoveTemp(Initializer));
 		DynamicRayTracingGeometry.InitResource(RHICmdList);
 	}
 
diff --git a/Engine/Plugins/Runtime/ProceduralMeshComponent/Source/ProceduralMeshComponent/Private/ProceduralMeshComponent.cpp b/Engine/Plugins/Runtime/ProceduralMeshComponent/Source/ProceduralMeshComponent/Private/ProceduralMeshComponent.cpp
--- a/Engine/Plugins/Runtime/ProceduralMeshComponent/Source/ProceduralMeshComponent/Private/ProceduralMeshComponent.cpp
+++ b/Engine/Plugins/Runtime/ProceduralMeshComponent/Source/ProceduralMeshComponent/Private/ProceduralMeshComponent.cpp
@@ -181,27 +181,21 @@ public:
 					{
 						FRayTracingGeometryInitializer Initializer;
 						Initializer.DebugName = DebugName;
-						Initializer.IndexBuffer = nullptr;
-						Initializer.TotalPrimitiveCount = 0;
+						Initializer.IndexBuffer = NewSection->IndexBuffer.IndexBufferRHI;
+						Initializer.TotalPrimitiveCount = NewSection->IndexBuffer.Indices.Num() / 3;
 						Initializer.GeometryType = RTGT_Triangles;
 						Initializer.bFastBuild = true;
 						Initializer.bAllowUpdate = false;
 
-						NewSection->RayTracingGeometry.SetInitializer(Initializer);
-						NewSection->RayTracingGeometry.InitResource(RHICmdList);
-
-						NewSection->RayTracingGeometry.Initializer.IndexBuffer = NewSection->IndexBuffer.IndexBufferRHI;
-						NewSection->RayTracingGeometry.Initializer.TotalPrimitiveCount = NewSection->IndexBuffer.Indices.Num() / 3;
-
 						FRayTracingGeometrySegment Segment;
 						Segment.VertexBuffer = NewSection->VertexBuffers.PositionVertexBuffer.VertexBufferRHI;
-						Segment.NumPrimitives = NewSection->RayTracingGeometry.Initializer.TotalPrimitiveCount;
 						Segment.MaxVertices = NewSection->VertexBuffers.PositionVertexBuffer.GetNumVertices();
-						NewSection->RayTracingGeometry.Initializer.Segments.Add(Segment);
-
-						//#dxr_todo: add support for segments?
+						Segment.NumPrimitives = Initializer.TotalPrimitiveCount;
 						
-						NewSection->RayTracingGeometry.UpdateRHI(RHICmdList);
+						Initializer.Segments.Add(Segment);
+
+						NewSection->RayTracingGeometry.SetInitializer(Initializer);
+						NewSection->RayTracingGeometry.InitResource(RHICmdList);
 					});
 				}
 #endif
@@ -306,16 +300,15 @@ public:
 					Initializer.bFastBuild = true;
 					Initializer.bAllowUpdate = false;
 
-					Section->RayTracingGeometry.SetInitializer(Initializer);
-					Section->RayTracingGeometry.InitResource(RHICmdList);
-
 					FRayTracingGeometrySegment Segment;
 					Segment.VertexBuffer = Section->VertexBuffers.PositionVertexBuffer.VertexBufferRHI;
-					Segment.NumPrimitives = Section->RayTracingGeometry.Initializer.TotalPrimitiveCount;
 					Segment.MaxVertices = Section->VertexBuffers.PositionVertexBuffer.GetNumVertices();
-					Section->RayTracingGeometry.Initializer.Segments.Add(Segment);
+					Segment.NumPrimitives = Initializer.TotalPrimitiveCount;
 
-					Section->RayTracingGeometry.UpdateRHI(RHICmdList);
+					Initializer.Segments.Add(Segment);
+
+					Section->RayTracingGeometry.SetInitializer(Initializer);
+					Section->RayTracingGeometry.InitResource(RHICmdList);
 				}
 #endif
 			}
diff --git a/Engine/Source/Runtime/Engine/Private/InstancedStaticMesh.cpp b/Engine/Source/Runtime/Engine/Private/InstancedStaticMesh.cpp
--- a/Engine/Source/Runtime/Engine/Private/InstancedStaticMesh.cpp
+++ b/Engine/Source/Runtime/Engine/Private/InstancedStaticMesh.cpp
@@ -2597,8 +2597,7 @@ void FInstancedStaticMeshSceneProxy::SetupRayTracingDynamicInstances(int32 NumDy
 		{
 			FRayTracingDynamicData &DynamicData = RayTracingDynamicData.AddDefaulted_GetRef();
 
-			auto& Initializer = DynamicData.DynamicGeometry.Initializer;
-			Initializer = LODModel.RayTracingGeometry.Initializer;
+			FRayTracingGeometryInitializer Initializer = LODModel.RayTracingGeometry.Initializer;
 			for (FRayTracingGeometrySegment& Segment : Initializer.Segments)
 			{
 				Segment.VertexBuffer = nullptr; 
@@ -2606,6 +2605,7 @@ void FInstancedStaticMeshSceneProxy::SetupRayTracingDynamicInstances(int32 NumDy
 			Initializer.bAllowUpdate = true;
 			Initializer.bFastBuild = true;
 
+			DynamicData.DynamicGeometry.SetInitializer(MoveTemp(Initializer));
 			DynamicData.DynamicGeometry.InitResource(FRHICommandListImmediate::Get());
 		}
 	}
diff --git a/Engine/Source/Runtime/Engine/Private/StaticMeshRender.cpp b/Engine/Source/Runtime/Engine/Private/StaticMeshRender.cpp
--- a/Engine/Source/Runtime/Engine/Private/StaticMeshRender.cpp
+++ b/Engine/Source/Runtime/Engine/Private/StaticMeshRender.cpp
@@ -745,14 +745,14 @@ bool FStaticMeshSceneProxy::GetMeshElement(
 #if RHI_RAYTRACING
 void FStaticMeshSceneProxy::CreateDynamicRayTracingGeometries(FRHICommandListBase& RHICmdList)
 {
+	check(bDynamicRayTracingGeometry && bNeedsDynamicRayTracingGeometries);
 	check(DynamicRayTracingGeometries.IsEmpty());
 
 	DynamicRayTracingGeometries.AddDefaulted(RenderData->LODResources.Num());
 
 	for (int32 LODIndex = 0; LODIndex < RenderData->LODResources.Num(); LODIndex++)
 	{
-		auto& Initializer = DynamicRayTracingGeometries[LODIndex].Initializer;
-		Initializer = RenderData->LODResources[LODIndex].RayTracingGeometry.Initializer;
+		FRayTracingGeometryInitializer Initializer = RenderData->LODResources[LODIndex].RayTracingGeometry.Initializer;
 		for (FRayTracingGeometrySegment& Segment : Initializer.Segments)
 		{
 			Segment.VertexBuffer = nullptr;
@@ -760,12 +760,9 @@ void FStaticMeshSceneProxy::CreateDynamicRayTracingGeometries(FRHICommandListBas
 		Initializer.bAllowUpdate = true;
 		Initializer.bFastBuild = true;
 		Initializer.Type = ERayTracingGeometryInitializerType::Rendering;
-	}
 
-	for (int32 i = 0; i < DynamicRayTracingGeometries.Num(); i++)
-	{
-		auto& Geometry = DynamicRayTracingGeometries[i];
-		Geometry.InitResource(RHICmdList);
+		DynamicRayTracingGeometries[LODIndex].SetInitializer(MoveTemp(Initializer));
+		DynamicRayTracingGeometries[LODIndex].InitResource(RHICmdList);
 	}
 }
 
@@ -785,7 +782,6 @@ void FStaticMeshSceneProxy::CreateRenderThreadResources(FRHICommandListBase& RHI
 #if RHI_RAYTRACING
 	if(IsRayTracingAllowed() && bNeedsDynamicRayTracingGeometries)
 	{
-		check(bDynamicRayTracingGeometry);
 		CreateDynamicRayTracingGeometries(RHICmdList);
 	}
 	else
diff --git a/Engine/Source/Runtime/GeometryFramework/Public/Components/BaseDynamicMeshSceneProxy.h b/Engine/Source/Runtime/GeometryFramework/Public/Components/BaseDynamicMeshSceneProxy.h
--- a/Engine/Source/Runtime/GeometryFramework/Public/Components/BaseDynamicMeshSceneProxy.h
+++ b/Engine/Source/Runtime/GeometryFramework/Public/Components/BaseDynamicMeshSceneProxy.h
@@ -342,16 +342,15 @@ protected:
 			Initializer.bFastBuild = true;
 			Initializer.bAllowUpdate = false;
 
-			RayTracingGeometry.SetInitializer(Initializer);
-			RayTracingGeometry.InitResource(RHICmdList);
-
 			FRayTracingGeometrySegment Segment;
 			Segment.VertexBuffer = PositionVertexBuffer.VertexBufferRHI;
-			Segment.NumPrimitives = RayTracingGeometry.Initializer.TotalPrimitiveCount;
+			Segment.NumPrimitives = Initializer.TotalPrimitiveCount;
 			Segment.MaxVertices = PositionVertexBuffer.GetNumVertices();
-			RayTracingGeometry.Initializer.Segments.Add(Segment);
 
-			RayTracingGeometry.UpdateRHI(RHICmdList);
+			Initializer.Segments.Add(Segment);
+
+			RayTracingGeometry.SetInitializer(MoveTemp(Initializer));
+			RayTracingGeometry.InitResource(RHICmdList);
 		}
 #endif
 	}
