ProceduralMeshComponent: Fix SliceProceduralMesh using incorrect normal in presence of non-uniform scale.

PR: #12298

diff --git a/Engine/Plugins/Runtime/ProceduralMeshComponent/Source/ProceduralMeshComponent/Private/KismetProceduralMeshLibrary.cpp b/Engine/Plugins/Runtime/ProceduralMeshComponent/Source/ProceduralMeshComponent/Private/KismetProceduralMeshLibrary.cpp
--- a/Engine/Plugins/Runtime/ProceduralMeshComponent/Source/ProceduralMeshComponent/Private/KismetProceduralMeshLibrary.cpp
+++ b/Engine/Plugins/Runtime/ProceduralMeshComponent/Source/ProceduralMeshComponent/Private/KismetProceduralMeshLibrary.cpp
@@ -793,7 +793,7 @@ void UKismetProceduralMeshLibrary::SliceProceduralMesh(UProceduralMeshComponent*
 		// Transform plane from world to local space
 		FTransform ProcCompToWorld = InProcMesh->GetComponentToWorld();
 		FVector LocalPlanePos = ProcCompToWorld.InverseTransformPosition(PlanePosition);
-		FVector LocalPlaneNormal = ProcCompToWorld.InverseTransformVectorNoScale(PlaneNormal);
+		FVector LocalPlaneNormal = ProcCompToWorld.GetScale3D() * ProcCompToWorld.InverseTransformVectorNoScale(PlaneNormal);
 		LocalPlaneNormal = LocalPlaneNormal.GetSafeNormal(); // Ensure normalized
 
 		FPlane SlicePlane(LocalPlanePos, LocalPlaneNormal);
